glTF importer work

Foundations of work for Skeletal Mesh import; right now we just support Static Mesh.
- node hierarchy
- joint IDs & skinning weights
- matrix & quaternion values

diff --git a/Engine/Plugins/Editor/GLTFImporter/Source/GLTFImporter/Private/GLTFAsset.cpp b/Engine/Plugins/Editor/GLTFImporter/Source/GLTFImporter/Private/GLTFAsset.cpp
--- a/Engine/Plugins/Editor/GLTFImporter/Source/GLTFImporter/Private/GLTFAsset.cpp
+++ b/Engine/Plugins/Editor/GLTFImporter/Source/GLTFImporter/Private/GLTFAsset.cpp
@@ -27,6 +27,7 @@ FString FTex::GetName() const
 }
 
 void FAsset::Reset(uint32 BufferCount, uint32 BufferViewCount, uint32 AccessorCount, uint32 MeshCount,
+                   uint32 NodeCount, uint32 SkinCount,
                    uint32 ImageCount, uint32 SamplerCount, uint32 TextureCount, uint32 MaterialCount)
 {
 	Buffers.Empty(BufferCount);
@@ -34,6 +35,9 @@ void FAsset::Reset(uint32 BufferCount, uint32 BufferViewCount, uint32 AccessorCo
 	Accessors.Empty(AccessorCount);
 	Meshes.Empty(MeshCount);
 
+	Nodes.Empty(NodeCount);
+	Skins.Empty(SkinCount);
+
 	Images.Empty(ImageCount);
 	Samplers.Empty(SamplerCount);
 	Textures.Empty(TextureCount);
diff --git a/Engine/Plugins/Editor/GLTFImporter/Source/GLTFImporter/Private/GLTFAsset.h b/Engine/Plugins/Editor/GLTFImporter/Source/GLTFImporter/Private/GLTFAsset.h
--- a/Engine/Plugins/Editor/GLTFImporter/Source/GLTFImporter/Private/GLTFAsset.h
+++ b/Engine/Plugins/Editor/GLTFImporter/Source/GLTFImporter/Private/GLTFAsset.h
@@ -18,6 +18,17 @@ struct UInt32Vector
 	uint32 A { 0 }, B { 0 }, C { 0 };
 };
 
+struct JointInfluence
+{
+	FVector4 Weight;
+	uint16 ID[4]{ 0, 0, 0, 0 };
+
+	JointInfluence(const FVector4& InWeight)
+		: Weight(InWeight)
+	{
+	}
+};
+
 // --- in-memory versions of glTF types -----------
 
 struct FBuffer
@@ -108,9 +119,11 @@ struct FAccessor
 	virtual bool IsValid() const = 0;
 
 	virtual uint32 GetUnsignedInt(uint32 Index) const = 0;
+	virtual void GetUnsignedInt16x4(uint32 Index, uint16 Values[4]) const = 0;
 	virtual FVector2D GetVec2(uint32 Index) const = 0;
 	virtual FVector GetVec3(uint32 Index) const = 0;
 	virtual FVector4 GetVec4(uint32 Index) const = 0;
+	virtual FMatrix GetMat4(uint32 Index) const = 0;
 
 	// GetMatrix functions are not needed yet, will be added later.
 
@@ -118,6 +131,7 @@ struct FAccessor
 	virtual TArray<FVector2D> GetVec2Array() const = 0;
 	virtual TArray<FVector> GetVec3Array() const = 0;
 	virtual TArray<FVector4> GetVec4Array() const = 0;
+	virtual TArray<FMatrix> GetMat4Array() const = 0;
 
 	virtual uint32 ValueSize() const = 0;
 };
@@ -189,6 +203,35 @@ public:
 		return 0;
 	}
 
+	void GetUnsignedInt16x4(uint32 Index, uint16 Values[4]) const override
+	{
+		// should be Vec4, not Normalized, unsigned integer (8 or 16 bit)
+
+		if (Index < Count)
+		{
+			if (Type == EType::Vec4 && !Normalized)
+			{
+				const void* Pointer = DataAt(Index);
+
+				switch (CompType)
+				{
+				case ECompType::U8:
+					for (int i = 0; i < 4; ++i)
+					{
+						Values[i] = *(const uint8*)Pointer;
+					}
+				case ECompType::U16:
+					for (int i = 0; i < 4; ++i)
+					{
+						Values[i] = *(const uint16*)Pointer;
+					}
+				default:
+					break;
+				}
+			}
+		}
+	}
+
 	FVector2D GetVec2(uint32 Index) const override
 	{
 		// Spec-defined attributes (TEXCOORD_0, TEXCOORD_1) use only these formats:
@@ -318,6 +361,44 @@ public:
 		return FVector4();
 	}
 
+	FMatrix GetMat4(uint32 Index) const override
+	{
+		// Focus on F32 for now, add other types as needed.
+
+		FMatrix Matrix;
+
+		if (Index < Count)
+		{
+			if (Type == EType::Mat4) // strict format match, unlike GPU shader fetch
+			{
+				const void* Pointer = DataAt(Index);
+
+				if (CompType == ECompType::F32)
+				{
+					// copy float vec4 directly from buffer
+					const float* P = static_cast<const float*>(Pointer);
+					for (int32 Row = 0; Row < 4; ++Row)
+					{
+						for (int32 Col = 0; Col < 4; ++Col)
+						{
+							// glTF stores matrix elements in column major order
+							// Unreal's FMatrix is row major
+							Matrix.M[Row][Col] = P[Col * 4 + Row];
+						}
+					}
+
+				}
+				else
+				{
+					check(false); // are normalized int types valid?
+				}
+			}
+		}
+
+		// unsupported format
+		return Matrix;
+	}
+
 	TArray<uint32> GetUnsignedIntArray() const override
 	{
 		// cheesy & unoptimized
@@ -365,6 +446,18 @@ public:
 		}
 		return Result;
 	}
+
+	TArray<FMatrix> GetMat4Array() const override
+	{
+		// cheesy & unoptimized
+		TArray<FMatrix> Result;
+		Result.Reserve(Count);
+		for (uint32 i = 0; i < Count; ++i)
+		{
+			Result.Push(GetMat4(i));
+		}
+		return Result;
+	}
 };
 
 struct FVoidAccessor final : FAccessor
@@ -384,6 +477,10 @@ struct FVoidAccessor final : FAccessor
 		return 0;
 	}
 
+	void GetUnsignedInt16x4(uint32 Index, uint16 Values[4]) const override
+	{
+	}
+
 	FVector2D GetVec2(uint32 Index) const override
 	{
 		return FVector2D::ZeroVector;
@@ -399,6 +496,11 @@ struct FVoidAccessor final : FAccessor
 		return FVector4();
 	}
 
+	FMatrix GetMat4(uint32 Index) const override
+	{
+		return FMatrix::Identity;
+	}
+
 	TArray<uint32> GetUnsignedIntArray() const override
 	{
 		return TArray<uint32>();
@@ -419,6 +521,11 @@ struct FVoidAccessor final : FAccessor
 		return TArray<FVector4>();
 	}
 
+	TArray<FMatrix> GetMat4Array() const override
+	{
+		return TArray<FMatrix>();
+	}
+
 	uint32 ValueSize() const override
 	{
 		return 0;
@@ -430,13 +537,16 @@ class FPrimitive
 	// index buffer
 	const FAccessor& Indices;
 
-	// attributes
+	// common attributes
 	const FAccessor& Position; // required
 	const FAccessor& Normal;
 	const FAccessor& Tangent;
 	const FAccessor& TexCoord0;
 	const FAccessor& TexCoord1;
 	const FAccessor& Color0;
+	// skeletal mesh attributes
+	const FAccessor& Joints0;
+	const FAccessor& Weights0;
 
 public:
 	enum class EMode
@@ -467,7 +577,9 @@ public:
 	           const FAccessor& InTangent,
 	           const FAccessor& InTexCoord0,
 	           const FAccessor& InTexCoord1,
-	           const FAccessor& InColor0)
+	           const FAccessor& InColor0,
+	           const FAccessor& InJoints0,
+	           const FAccessor& InWeights0)
 		: Indices(InIndices)
 		, Position(InPosition)
 		, Normal(InNormal)
@@ -475,8 +587,10 @@ public:
 		, TexCoord0(InTexCoord0)
 		, TexCoord1(InTexCoord1)
 		, Color0(InColor0)
+		, Joints0(InJoints0)
+		, Weights0(InWeights0)
 		, Mode(InMode)
-		, MaterialIndex(InMaterial)	
+		, MaterialIndex(InMaterial)
 	{
 	}
 
@@ -533,6 +647,26 @@ public:
 		}
 	}
 
+	bool HasJointWeights() const
+	{
+		return Joints0.IsValid() && Weights0.IsValid();
+	}
+
+	TArray<JointInfluence> GetJointInfluences() const
+	{
+		TArray<JointInfluence> Influences;
+		// return a flat array that corresponds 1-to-1 with vertex positions
+		const int32 N = Joints0.Count;
+		Influences.Reserve(N);
+		for (int32 i = 0; i < N; ++i)
+		{
+			JointInfluence& Joint = Influences.Emplace_GetRef(Weights0.GetVec4(i));
+			Joints0.GetUnsignedInt16x4(i, Joint.ID);
+		}
+
+		return Influences;
+	}
+
 	UInt32Vector TriangleVerts(uint32 T) const
 	{
 		UInt32Vector Result;
@@ -711,6 +845,52 @@ struct FMesh
 		}
 		return false;
 	}
+
+	bool HasJointWeights() const
+	{
+		// According to spec, *all* primitives of a skinned mesh must have joint weights.
+		// TODO: add this to validation
+		for (const FPrimitive& Prim : Primitives)
+		{
+			if (Prim.HasJointWeights())
+			{
+				return true;
+			}
+		}
+		return false;
+	}
+};
+
+struct FNode
+{
+	FString Name;
+	FTransform Transform;
+
+	TArray<int32> Children; // index into Asset::Nodes
+
+	bool IsJoint { false };
+	// TODO: enum, whether this node is a Mesh, Joint, Camera, Empty, etc.
+	// more accurately, does this node *have* a Mesh, Camera, etc.
+
+	int32 MeshIndex { INDEX_NONE }; // index into Asset::Meshes
+	int32 Skindex { INDEX_NONE }; // index into Asset::Skins
+
+	// Skindex is the skin used by the mesh at this node.
+	// It's not the joints belonging *to* a skin 
+};
+
+struct FSkinfo
+{
+	FString Name;
+	TArray<int32> Joints; // each is an index into Asset::Nodes
+	int32 Skeleton { INDEX_NONE }; // AKA root node, index into Asset::Nodes
+
+	const FAccessor& InverseBindMatrices;
+
+	FSkinfo(const FAccessor& InIBM)
+		: InverseBindMatrices(InIBM)
+	{
+	}
 };
 
 struct FImage
@@ -857,12 +1037,16 @@ struct FAsset
 	TArray<FValidAccessor> Accessors;
 	TArray<FMesh> Meshes;
 
+	TArray<FNode> Nodes;
+	TArray<FSkinfo> Skins;
+
 	TArray<FImage> Images;
 	TArray<FSampler> Samplers;
 	TArray<FTex> Textures;
 	TArray<FMat> Materials;
 
 	void Reset(uint32 BufferCount, uint32 BufferViewCount, uint32 AccessorCount, uint32 MeshCount,
+	           uint32 NodeCount, uint32 SkinCount,
 	           uint32 ImageCount, uint32 SamplerCount, uint32 TextureCount, uint32 MaterialCount);
 
 	bool Validate() const;
diff --git a/Engine/Plugins/Editor/GLTFImporter/Source/GLTFImporter/Private/GLTFReader.cpp b/Engine/Plugins/Editor/GLTFImporter/Source/GLTFImporter/Private/GLTFReader.cpp
--- a/Engine/Plugins/Editor/GLTFImporter/Source/GLTFImporter/Private/GLTFReader.cpp
+++ b/Engine/Plugins/Editor/GLTFImporter/Source/GLTFImporter/Private/GLTFReader.cpp
@@ -224,6 +224,49 @@ static FVector4 GetVec4(const FJsonObject& Object, const char* Name, const FVect
 	return DefaultValue;
 }
 
+static FQuat GetQuat(const FJsonObject& Object, const char* Name, const FQuat& DefaultValue = FQuat(0,0,0,1))
+{
+	if (Object.HasTypedField<EJson::Array>(Name))
+	{
+		const TArray<TSharedPtr<FJsonValue>>& Array = Object.GetArrayField(Name);
+		if (Array.Num() == 4)
+		{
+			float X = Array[0]->AsNumber();
+			float Y = Array[1]->AsNumber();
+			float Z = Array[2]->AsNumber();
+			float W = Array[3]->AsNumber();
+			return FQuat(X, Y, Z, W);
+		}
+	}
+
+	return DefaultValue;
+}
+
+static FMatrix GetMat4(const FJsonObject& Object, const char* Name, const FMatrix& DefaultValue = FMatrix::Identity)
+{
+	if (Object.HasTypedField<EJson::Array>(Name))
+	{
+		const TArray<TSharedPtr<FJsonValue>>& Array = Object.GetArrayField(Name);
+		if (Array.Num() == 16)
+		{
+			FMatrix Matrix;
+			for (int32 Row = 0; Row < 4; ++Row)
+			{
+				for (int32 Col = 0; Col < 4; ++Col)
+				{
+					// glTF stores matrix elements in column major order
+					// Unreal's FMatrix is row major
+					Matrix.M[Row][Col] = Array[Col * 4 + Row]->AsNumber();
+				}
+			}
+
+			return Matrix;
+		}
+	}
+
+	return DefaultValue;
+}
+
 // --- Base64 utility function --------------------------
 
 static bool DecodeDataURI(const FString& URI, FString& OutMimeType, TArray<uint8>& OutData)
@@ -377,8 +420,10 @@ void FGLTFReader::SetupPrimitive(const FJsonObject& Object, FMesh& Mesh)
 	const FAccessor& TexCoord0 = AccessorAtIndex(A, GetIndex(Attributes, "TEXCOORD_0"));
 	const FAccessor& TexCoord1 = AccessorAtIndex(A, GetIndex(Attributes, "TEXCOORD_1"));
 	const FAccessor& Color0 = AccessorAtIndex(A, GetIndex(Attributes, "COLOR_0"));
+	const FAccessor& Joints0 = AccessorAtIndex(A, GetIndex(Attributes, "JOINTS_0"));
+	const FAccessor& Weights0 = AccessorAtIndex(A, GetIndex(Attributes, "WEIGHTS_0"));
 
-	Mesh.Primitives.Emplace(Mode, MaterialIndex, Indices, Position, Normal, Tangent, TexCoord0, TexCoord1, Color0);
+	Mesh.Primitives.Emplace(Mode, MaterialIndex, Indices, Position, Normal, Tangent, TexCoord0, TexCoord1, Color0, Joints0, Weights0);
 	if (!Mesh.Primitives.Last().Validate())
 	{
 		Warn.Log("Invalid primitive :(");
@@ -402,6 +447,54 @@ void FGLTFReader::SetupMesh(const FJsonObject& Object)
 	}
 }
 
+void FGLTFReader::SetupNode(const FJsonObject& Object)
+{
+	FNode& Node = Asset.Nodes.Emplace_GetRef();
+
+	Node.Name = GetString(Object, "name");
+
+	if (Object.HasField("matrix"))
+	{
+		Node.Transform.SetFromMatrix(GetMat4(Object, "matrix"));
+	}
+	else
+	{
+		Node.Transform.SetTranslation(GetVec3(Object, "translation"));
+		Node.Transform.SetRotation(GetQuat(Object, "rotation"));
+		Node.Transform.SetScale3D(GetVec3(Object, "scale", FVector::OneVector));
+	}
+
+	const TArray<TSharedPtr<FJsonValue>>& ChildArray = Object.GetArrayField("children");
+	Node.Children.Reserve(ChildArray.Num());
+	for (TSharedPtr<FJsonValue> Value : ChildArray)
+	{
+		const int32 ChildIndex = Value->AsNumber();
+		Node.Children.Add(ChildIndex);
+	}
+
+	Node.MeshIndex = GetIndex(Object, "mesh");
+	Node.Skindex = GetIndex(Object, "skin");
+}
+
+void FGLTFReader::SetupSkin(const FJsonObject& Object)
+{
+	const FAccessor& InverseBindMatrices = AccessorAtIndex(Asset.Accessors, GetIndex(Object, "inverseBindMatrices"));
+
+	FSkinfo& Skin = Asset.Skins.Emplace_GetRef(InverseBindMatrices);
+
+	Skin.Name = GetString(Object, "name");
+
+	const TArray<TSharedPtr<FJsonValue>>& JointArray = Object.GetArrayField("joints");
+	Skin.Joints.Reserve(JointArray.Num());
+	for (TSharedPtr<FJsonValue> Value : JointArray)
+	{
+		const int32 JointIndex = Value->AsNumber();
+		Skin.Joints.Add(JointIndex);
+	}
+
+	Skin.Skeleton = GetIndex(Object, "skeleton");
+}
+
 void FGLTFReader::SetupImage(const FJsonObject& Object)
 {
 	Asset.Images.Emplace();
@@ -586,12 +679,16 @@ FGLTFReader::FGLTFReader(FAsset& OutAsset, const FJsonObject& Root, const FStrin
 	AccessorCount = ArraySize(Root, "accessors");
 	MeshCount = ArraySize(Root, "meshes");
 
+	NodeCount = ArraySize(Root, "nodes");
+	SkinCount = ArraySize(Root, "skins");
+
 	ImageCount = ArraySize(Root, "images");
 	SamplerCount = ArraySize(Root, "samplers");
 	TextureCount = ArraySize(Root, "textures");
 	MaterialCount = ArraySize(Root, "materials");
 
 	Asset.Reset(BufferCount, BufferViewCount, AccessorCount, MeshCount,
+	            NodeCount, SkinCount,
 	            ImageCount, SamplerCount, TextureCount, MaterialCount);
 
 	if (BufferCount > 0)
@@ -630,6 +727,24 @@ FGLTFReader::FGLTFReader(FAsset& OutAsset, const FJsonObject& Root, const FStrin
 		}
 	}
 
+	if (NodeCount > 0)
+	{
+		for (TSharedPtr<FJsonValue> Value : Root.GetArrayField("nodes"))
+		{
+			const FJsonObject& Object = *Value->AsObject();
+			SetupNode(Object);
+		}
+	}
+
+	if (SkinCount > 0)
+	{
+		for (TSharedPtr<FJsonValue> Value : Root.GetArrayField("skins"))
+		{
+			const FJsonObject& Object = *Value->AsObject();
+			SetupSkin(Object);
+		}
+	}
+
 	if (ImageCount > 0)
 	{
 		for (TSharedPtr<FJsonValue> Value : Root.GetArrayField("images"))
diff --git a/Engine/Plugins/Editor/GLTFImporter/Source/GLTFImporter/Private/GLTFReader.h b/Engine/Plugins/Editor/GLTFImporter/Source/GLTFImporter/Private/GLTFReader.h
--- a/Engine/Plugins/Editor/GLTFImporter/Source/GLTFImporter/Private/GLTFReader.h
+++ b/Engine/Plugins/Editor/GLTFImporter/Source/GLTFImporter/Private/GLTFReader.h
@@ -18,6 +18,9 @@ class FGLTFReader
 	uint32 AccessorCount { 0 };
 	uint32 MeshCount { 0 };
 
+	uint32 NodeCount { 0 };
+	uint32 SkinCount { 0 };
+
 	uint32 ImageCount { 0 };
 	uint32 SamplerCount { 0 };
 	uint32 TextureCount { 0 };
@@ -41,6 +44,9 @@ class FGLTFReader
 	void SetupPrimitive(const FJsonObject&, GLTF::FMesh&);
 	void SetupMesh(const FJsonObject&);
 
+	void SetupNode(const FJsonObject&);
+	void SetupSkin(const FJsonObject&);
+
 	void SetupImage(const FJsonObject&);
 	void SetupSampler(const FJsonObject&);
 	void SetupTexture(const FJsonObject&);
diff --git a/Engine/Plugins/Editor/GLTFImporter/Source/GLTFImporter/Private/GLTFStaticMesh.cpp b/Engine/Plugins/Editor/GLTFImporter/Source/GLTFImporter/Private/GLTFStaticMesh.cpp
--- a/Engine/Plugins/Editor/GLTFImporter/Source/GLTFImporter/Private/GLTFStaticMesh.cpp
+++ b/Engine/Plugins/Editor/GLTFImporter/Source/GLTFImporter/Private/GLTFStaticMesh.cpp
@@ -161,6 +161,11 @@ UStaticMesh* ImportStaticMesh(const FAsset& Asset, const TArray<UMaterial*>& Mat
 
 	const FMesh& Mesh = Asset.Meshes[Index];
 
+	if (Mesh.HasJointWeights())
+	{
+		UE_LOG(LogTemp, Warning, TEXT("Mesh has joint weights; import as Skeletal Mesh?"));
+	}
+
 	FString AssetName;
 	UPackage* AssetPackage = GetAssetPackageAndName<UStaticMesh>(InParent, Mesh.Name, TEXT("SM"), InName, Index, AssetName);
 
