Deprecate FRawMesh - GLTF importer

diff --git a/Engine/Plugins/Editor/GLTFImporter/Source/GLTFImporter/GLTFImporter.Build.cs b/Engine/Plugins/Editor/GLTFImporter/Source/GLTFImporter/GLTFImporter.Build.cs
--- a/Engine/Plugins/Editor/GLTFImporter/Source/GLTFImporter/GLTFImporter.Build.cs
+++ b/Engine/Plugins/Editor/GLTFImporter/Source/GLTFImporter/GLTFImporter.Build.cs
@@ -29,7 +29,7 @@ namespace UnrealBuildTool.Rules
 					"CoreUObject",
 					"Engine",
 					"UnrealEd",
-					"RawMesh",
+					"MeshDescription",
 					"MeshUtilities",
                     "MessageLog",
 					"Json",
diff --git a/Engine/Plugins/Editor/GLTFImporter/Source/GLTFImporter/Private/GLTFStaticMesh.cpp b/Engine/Plugins/Editor/GLTFImporter/Source/GLTFImporter/Private/GLTFStaticMesh.cpp
--- a/Engine/Plugins/Editor/GLTFImporter/Source/GLTFImporter/Private/GLTFStaticMesh.cpp
+++ b/Engine/Plugins/Editor/GLTFImporter/Source/GLTFImporter/Private/GLTFStaticMesh.cpp
@@ -4,7 +4,10 @@
 #include "GLTFPackage.h"
 
 #include "Engine/StaticMesh.h"
-#include "RawMesh.h"
+//#include "RawMesh.h"
+#include "MeshDescription.h"
+#include "MeshAttributes.h"
+#include "MeshAttributeArray.h"
 #include "Materials/Material.h"
 #include "AssetRegistryModule.h"
 
@@ -57,18 +60,6 @@ static TArray<FVector> Tanslate(const TArray<FVector4>& In)
 	return Result;
 }
 
-static void AppendWedgeIndices(FRawMesh& RawMesh, const TArray<uint32>& PrimIndices)
-{
-	// Indices come in Prim-relative, we need them Mesh-relative.
-	// Simply offset each index by the number of vertices already in the mesh.
-	// Call this *before* adding this Prim's vertex data.
-	const uint32 Offset = RawMesh.VertexPositions.Num();
-	RawMesh.WedgeIndices.Reserve(RawMesh.WedgeIndices.Num() + PrimIndices.Num());
-	for (uint32 Index : PrimIndices)
-	{
-		RawMesh.WedgeIndices.Add(Index + Offset);
-	}
-}
 
 static TArray<FVector> GenerateFlatNormals(const TArray<FVector>& Positions, const TArray<uint32>& Indices)
 {
@@ -105,7 +96,7 @@ static void AddN(TArray<T>& Array, T Value, uint32 N)
 	}
 }
 
-static void AssignMaterials(UStaticMesh* StaticMesh, FRawMesh& RawMesh, const TArray<UMaterial*>& Materials, const TSet<int32>& MaterialIndices)
+static void AssignMaterials(UStaticMesh* StaticMesh, TMap<int32, int32>& OutMaterialIndexToSlot, const TArray<UMaterial*>& Materials, const TSet<int32>& MaterialIndices)
 {
 	// Create material slots for this mesh, only for the materials it uses.
 
@@ -115,9 +106,7 @@ static void AssignMaterials(UStaticMesh* StaticMesh, FRawMesh& RawMesh, const TA
 	SortedMaterialIndices.Sort();
 
 	const int32 N = MaterialIndices.Num();
-
-	TMap<int32, int32> MaterialIndexToSlot;
-	MaterialIndexToSlot.Reserve(N);
+	OutMaterialIndexToSlot.Empty(N);
 	StaticMesh->StaticMaterials.Reserve(N);
 
 	for (int32 MaterialIndex : SortedMaterialIndices)
@@ -140,16 +129,10 @@ static void AssignMaterials(UStaticMesh* StaticMesh, FRawMesh& RawMesh, const TA
 			MeshSlot = StaticMesh->StaticMaterials.Emplace(Mat, MatName, MatName);
 		}
 
-		MaterialIndexToSlot.Add(MaterialIndex, MeshSlot);
+		OutMaterialIndexToSlot.Add(MaterialIndex, MeshSlot);
 
 		StaticMesh->SectionInfoMap.Set(0, MeshSlot, FMeshSectionInfo(MeshSlot));
 	}
-
-	// Replace material index references with this mesh's slots.
-	for (int32& Index : RawMesh.FaceMaterialIndices)
-	{
-		Index = MaterialIndexToSlot[Index];
-	}
 }
 
 UStaticMesh* ImportStaticMesh(const FAsset& Asset, const TArray<UMaterial*>& Materials, UObject* InParent, FName InName, EObjectFlags Flags, uint32 Index)
@@ -172,6 +155,37 @@ UStaticMesh* ImportStaticMesh(const FAsset& Asset, const TArray<UMaterial*>& Mat
 	UStaticMesh* StaticMesh = NewObject<UStaticMesh>(AssetPackage, FName(*AssetName), Flags);
 
 	FStaticMeshSourceModel& SourceModel = StaticMesh->AddSourceModel();
+	//GLTF do not support LOD yet so assuming LODIndex of 0
+	int32 LODIndex = 0;
+	FMeshDescription* MeshDescription = StaticMesh->CreateOriginalMeshDescription(LODIndex);
+	StaticMesh->RegisterMeshAttributes(*MeshDescription);
+
+	TVertexAttributeArray<FVector>& VertexPositions = MeshDescription->VertexAttributes().GetAttributes<FVector>(MeshAttribute::Vertex::Position);
+	TEdgeAttributeArray<bool>& EdgeHardnesses = MeshDescription->EdgeAttributes().GetAttributes<bool>(MeshAttribute::Edge::IsHard);
+	TEdgeAttributeArray<float>& EdgeCreaseSharpnesses = MeshDescription->EdgeAttributes().GetAttributes<float>(MeshAttribute::Edge::CreaseSharpness);
+	TPolygonGroupAttributeArray<FName>& PolygonGroupImportedMaterialSlotNames = MeshDescription->PolygonGroupAttributes().GetAttributes<FName>(MeshAttribute::PolygonGroup::ImportedMaterialSlotName);
+	TVertexInstanceAttributeArray<FVector>& VertexInstanceNormals = MeshDescription->VertexInstanceAttributes().GetAttributes<FVector>(MeshAttribute::VertexInstance::Normal);
+	TVertexInstanceAttributeArray<FVector>& VertexInstanceTangents = MeshDescription->VertexInstanceAttributes().GetAttributes<FVector>(MeshAttribute::VertexInstance::Tangent);
+	TVertexInstanceAttributeArray<float>& VertexInstanceBinormalSigns = MeshDescription->VertexInstanceAttributes().GetAttributes<float>(MeshAttribute::VertexInstance::BinormalSign);
+	//TVertexInstanceAttributeArray<FVector4>& VertexInstanceColors = MeshDescription->VertexInstanceAttributes().GetAttributes<FVector4>(MeshAttribute::VertexInstance::Color);
+	TVertexInstanceAttributeIndicesArray<FVector2D>& VertexInstanceUVs = MeshDescription->VertexInstanceAttributes().GetAttributesSet<FVector2D>(MeshAttribute::VertexInstance::TextureCoordinate);
+
+	int32 NumUVs = 0;
+	for (int32 UVIndex = 0; UVIndex < MAX_MESH_TEXTURE_COORDS_MD; ++UVIndex)
+	{
+		if (Mesh.HasTexCoords(UVIndex))
+		{
+			NumUVs++;
+		}
+		else
+		{
+			break;
+		}
+	}
+
+	VertexInstanceUVs.SetNumIndices(NumUVs);
+
+
 	FMeshBuildSettings& Settings = SourceModel.BuildSettings;
 
 	Settings.bRecomputeNormals = false;
@@ -187,94 +201,153 @@ UStaticMesh* ImportStaticMesh(const FAsset& Asset, const TArray<UMaterial*>& Mat
 
 	Settings.bGenerateLightmapUVs = false; // set to true if asset has no UV1 ?
 
-	FRawMesh RawMesh;
 	TSet<int32> MaterialIndicesUsed;
-
-	for (const FPrimitive& Prim : Mesh.Primitives)
+	//Add the vertex
+	TArray<TMap<int32, FVertexID>> PositionIndexToVertexID_PerPrim;
+	PositionIndexToVertexID_PerPrim.AddDefaulted(Mesh.Primitives.Num());
+	for (int32 PrimIndex = 0; PrimIndex < Mesh.Primitives.Num(); ++PrimIndex)
 	{
-		const uint32 TriCount = Prim.TriangleCount();
-		const TArray<uint32> Indices = Prim.GetTriangleIndices();
-
+		const FPrimitive& Prim = Mesh.Primitives[PrimIndex];
 		// Remember which primitives use which materials.
 		MaterialIndicesUsed.Add(Prim.MaterialIndex);
-		AddN(RawMesh.FaceMaterialIndices, Prim.MaterialIndex, TriCount);
-
-		// Make all faces part of the same smoothing group, so Unreal will combine identical adjacent verts.
-		constexpr uint32 SmoothMask = 1;
-		AddN(RawMesh.FaceSmoothingMasks, SmoothMask, TriCount);
-		// (Is there a way to set auto-gen smoothing threshold? glTF spec says to generate flat normals if they're not specified.
-		//   We want to combine identical verts whether they're smooth neighbors or triangles belonging to the same flat polygon.)
-
-		AppendWedgeIndices(RawMesh, Indices);
 
 		const TArray<FVector>& Positions = Prim.GetPositions();
-		RawMesh.VertexPositions.Append(Positions);
+		PositionIndexToVertexID_PerPrim[PrimIndex].Reserve(Positions.Num());
+		for (int32 PositionIndex = 0; PositionIndex < Positions.Num(); ++PositionIndex)
+		{
+			const FVertexID& VertexID = MeshDescription->CreateVertex();
+			VertexPositions[VertexID] = ConvertPosition(Positions[PositionIndex]);
+			PositionIndexToVertexID_PerPrim[PrimIndex].Add(PositionIndex, VertexID);
+		}
+	}
+	TMap<int32, int32> MaterialIndexToSlot;
+	AssignMaterials(StaticMesh, MaterialIndexToSlot, Materials, MaterialIndicesUsed);
 
+	TMap<int32, FPolygonGroupID> MaterialIndexToPolygonGroupID;
+	//Add the PolygonGroup
+	for (int32 MaterialIndex : MaterialIndicesUsed)
+	{
+		const FPolygonGroupID& PolygonGroupID = MeshDescription->CreatePolygonGroup();
+		MaterialIndexToPolygonGroupID.Add(MaterialIndex, PolygonGroupID);
+		PolygonGroupImportedMaterialSlotNames[PolygonGroupID] = StaticMesh->StaticMaterials[MaterialIndexToSlot[MaterialIndex]].ImportedMaterialSlotName;
+	}
+	//Add the vertexInstance
+	for (int32 PrimIndex = 0; PrimIndex < Mesh.Primitives.Num(); ++PrimIndex)
+	{
+		const FPrimitive& Prim = Mesh.Primitives[PrimIndex];
+		FPolygonGroupID CurrentPolygonGroupID = MaterialIndexToPolygonGroupID[Prim.MaterialIndex];
+		const uint32 TriCount = Prim.TriangleCount();
+		const TArray<uint32> Indices = Prim.GetTriangleIndices();
+		const TArray<FVector>& Positions = Prim.GetPositions();
+		TArray<FVector> Normals;
 		// glTF does not guarantee each primitive within a mesh has the same attributes.
 		// Fill in gaps as needed:
 		// - missing normals will be flat, based on triangle orientation
 		// - missing UVs will be (0,0)
 		// - missing tangents will be (0,0,1)
-
 		if (Prim.HasNormals())
 		{
-			RawMesh.WedgeTangentZ.Append(ReIndex(Prim.GetNormals(), Indices));
+			Normals = ReIndex(Prim.GetNormals(), Indices);
 		}
 		else
 		{
-			RawMesh.WedgeTangentZ.Append(GenerateFlatNormals(Positions, Indices));
+			Normals = GenerateFlatNormals(Positions, Indices);
 		}
-
+		TArray<FVector> Tangents;
 		if (Prim.HasTangents())
 		{
 			// glTF stores tangent as Vec4, with W component indicating handedness of tangent basis.
-			const TArray<FVector> Tangents = Tanslate(Prim.GetTangents());
-			RawMesh.WedgeTangentX.Append(ReIndex(Tangents, Indices));
+			Tangents = ReIndex(Tanslate(Prim.GetTangents()), Indices);
 		}
 		else if (Mesh.HasTangents())
 		{
 			// If other primitives in this mesh have tangents, generate filler ones for this primitive, to avoid gaps.
-			AddN(RawMesh.WedgeTangentX, FVector(0.0f, 0.0f, 1.0f), Prim.VertexCount());
+			AddN(Tangents, FVector(0.0f, 0.0f, 1.0f), Prim.VertexCount());
 			bDidGenerateTangents = true;
 		}
-
-		for (uint32 UV : { 0, 1 })
+		TArray<FVector2D> UVs[MAX_MESH_TEXTURE_COORDS_MD];
+		for (int32 UVIndex = 0; UVIndex < NumUVs; ++UVIndex)
 		{
-			if (Prim.HasTexCoords(UV))
+			if (Prim.HasTexCoords(UVIndex))
 			{
-				RawMesh.WedgeTexCoords[UV].Append(ReIndex(Prim.GetTexCoords(UV), Indices));
+				UVs[UVIndex] = ReIndex(Prim.GetTexCoords(UVIndex), Indices);
 			}
-			else if (UV == 0 || Mesh.HasTexCoords(UV))
+			else
 			{
 				// Unreal StaticMesh must have UV channel 0.
 				// glTF doesn't require this since not all materials need texture coordinates.
-				// We also fill UV channel >= 1 for this primitive if other primitives have it, to avoid gaps.
-				RawMesh.WedgeTexCoords[UV].AddZeroed(Prim.VertexCount());
+				// We also fill UV channel > 1 for this primitive if other primitives have it, to avoid gaps.
+				UVs[UVIndex].AddZeroed(Prim.VertexCount());
 				bDidGenerateTexCoords = true;
 			}
 		}
-	}
-
-	// Transform from glTF -> UE4 coordinate system.
-	// Tangents already transformed during conversion.
 
-	for (FVector& Pos : RawMesh.VertexPositions)
-	{
-		Pos = ConvertPosition(Pos);
-	}
+		//Now add all vertexInstances
+		for (uint32 TriangleIndex = 0; TriangleIndex < TriCount; ++TriangleIndex)
+		{
+			FVertexInstanceID CornerVertexInstanceIDs[3];
+			FVertexID CornerVertexIDs[3];
+			for (int32 Corner = 0; Corner < 3; ++Corner)
+			{
+				uint32 IndiceIndex = TriangleIndex * 3 + Corner;
+				int32 Index = Indices[IndiceIndex];
+
+				FVertexID VertexID = PositionIndexToVertexID_PerPrim[PrimIndex][Index];
+				const FVertexInstanceID& VertexInstanceID = MeshDescription->CreateVertexInstance(VertexID);
+
+				VertexInstanceTangents[VertexInstanceID] = Tangents[IndiceIndex];
+				VertexInstanceNormals[VertexInstanceID] = ConvertVec3(Normals[IndiceIndex]);
+				VertexInstanceBinormalSigns[VertexInstanceID] = GetBasisDeterminantSign(VertexInstanceTangents[VertexInstanceID].GetSafeNormal(),
+																						(VertexInstanceNormals[VertexInstanceID] ^ VertexInstanceTangents[VertexInstanceID]).GetSafeNormal(),
+																						VertexInstanceNormals[VertexInstanceID].GetSafeNormal());
+				for (int32 UVIndex = 0; UVIndex < NumUVs; ++UVIndex)
+				{
+					VertexInstanceUVs.GetArrayForIndex(UVIndex)[VertexInstanceID] = UVs[UVIndex][IndiceIndex];
+				}
+
+				CornerVertexInstanceIDs[Corner] = VertexInstanceID;
+				CornerVertexIDs[Corner] = VertexID;
+			}
 
-	for (FVector& Normal : RawMesh.WedgeTangentZ)
-	{
-		Normal = ConvertVec3(Normal);
+			TArray<FMeshDescription::FContourPoint> Contours;
+			for (int32 Corner = 0; Corner < 3; ++Corner)
+			{
+				int32 ContourPointIndex = Contours.AddDefaulted();
+				FMeshDescription::FContourPoint& ContourPoint = Contours[ContourPointIndex];
+				//Find the matching edge ID
+				uint32 CornerIndices[2];
+				CornerIndices[0] = (Corner + 0) % 3;
+				CornerIndices[1] = (Corner + 1) % 3;
+
+				FVertexID EdgeVertexIDs[2];
+				EdgeVertexIDs[0] = CornerVertexIDs[CornerIndices[0]];
+				EdgeVertexIDs[1] = CornerVertexIDs[CornerIndices[1]];
+
+				FEdgeID MatchEdgeId = MeshDescription->GetVertexPairEdge(EdgeVertexIDs[0], EdgeVertexIDs[1]);
+				if (MatchEdgeId == FEdgeID::Invalid)
+				{
+					MatchEdgeId = MeshDescription->CreateEdge(EdgeVertexIDs[0], EdgeVertexIDs[1]);
+					// Make all faces part of the same smoothing group, so Unreal will combine identical adjacent verts.
+					// (Is there a way to set auto-gen smoothing threshold? glTF spec says to generate flat normals if they're not specified.
+					//   We want to combine identical verts whether they're smooth neighbors or triangles belonging to the same flat polygon.)
+					EdgeHardnesses[MatchEdgeId] = false;
+					EdgeCreaseSharpnesses[MatchEdgeId] = 0.0f;
+				}
+				ContourPoint.EdgeID = MatchEdgeId;
+				ContourPoint.VertexInstanceID = CornerVertexInstanceIDs[CornerIndices[0]];
+			}
+			// Insert a polygon into the mesh
+			const FPolygonID NewPolygonID = MeshDescription->CreatePolygon(CurrentPolygonGroupID, Contours);
+			//Triangulate the polygon
+			FMeshPolygon& Polygon = MeshDescription->GetPolygon(NewPolygonID);
+			MeshDescription->ComputePolygonTriangulation(NewPolygonID, Polygon.Triangles);
+		}
 	}
-
-	// Finish processing materials for this mesh.
-	AssignMaterials(StaticMesh, RawMesh, Materials, MaterialIndicesUsed);
+	
 	// RawMesh.CompactMaterialIndices(); // needed?
 	bMeshUsesEmptyMaterial = MaterialIndicesUsed.Contains(INDEX_NONE);
 
-	SourceModel.RawMeshBulkData->SaveRawMesh(RawMesh);
-	StaticMesh->ImportVersion = EImportStaticMeshVersion::BeforeImportStaticMeshVersionWasAdded;
+	StaticMesh->CommitOriginalMeshDescription(LODIndex);
 	StaticMesh->PostEditChange();
 
 	// Set the dirty flag so this package will get saved later
diff --git a/Engine/Source/Developer/MeshBuilder/Private/MeshDescriptionHelper.h b/Engine/Source/Developer/MeshBuilder/Private/MeshDescriptionHelper.h
--- a/Engine/Source/Developer/MeshBuilder/Private/MeshDescriptionHelper.h
+++ b/Engine/Source/Developer/MeshBuilder/Private/MeshDescriptionHelper.h
@@ -13,12 +13,6 @@ struct FMeshBuildSettings;
 struct FVertexInstanceID;
 struct FMeshReductionSettings;
 
-enum
-{
-	//Remove the _MD when FRawMesh will be remove
-	MAX_MESH_TEXTURE_COORDS_MD = 8,
-};
-
 class FMeshDescriptionHelper
 {
 public:
diff --git a/Engine/Source/Runtime/MeshDescription/Public/MeshDescription.h b/Engine/Source/Runtime/MeshDescription/Public/MeshDescription.h
--- a/Engine/Source/Runtime/MeshDescription/Public/MeshDescription.h
+++ b/Engine/Source/Runtime/MeshDescription/Public/MeshDescription.h
@@ -11,6 +11,11 @@
 #include "UObject/EditorObjectVersion.h"
 #include "MeshDescription.generated.h"
 
+enum
+{
+	//Remove the _MD when FRawMesh will be remove
+	MAX_MESH_TEXTURE_COORDS_MD = 8,
+};
 
 struct FMeshVertex
 {
