Fix CIS build warning

diff --git a/Engine/Plugins/Editor/GLTFImporter/Source/GLTFImporter/Private/GLTFStaticMesh.cpp b/Engine/Plugins/Editor/GLTFImporter/Source/GLTFImporter/Private/GLTFStaticMesh.cpp
--- a/Engine/Plugins/Editor/GLTFImporter/Source/GLTFImporter/Private/GLTFStaticMesh.cpp
+++ b/Engine/Plugins/Editor/GLTFImporter/Source/GLTFImporter/Private/GLTFStaticMesh.cpp
@@ -290,9 +290,9 @@ UStaticMesh* ImportStaticMesh(const FAsset& Asset, const TArray<UMaterial*>& Mat
 			for (int32 Corner = 0; Corner < 3; ++Corner)
 			{
 				uint32 IndiceIndex = TriangleIndex * 3 + Corner;
-				int32 Index = Indices[IndiceIndex];
+				int32 VertexIndex = Indices[IndiceIndex];
 
-				FVertexID VertexID = PositionIndexToVertexID_PerPrim[PrimIndex][Index];
+				FVertexID VertexID = PositionIndexToVertexID_PerPrim[PrimIndex][VertexIndex];
 				const FVertexInstanceID& VertexInstanceID = MeshDescription->CreateVertexInstance(VertexID);
 
 				VertexInstanceTangents[VertexInstanceID] = Tangents[IndiceIndex];
