Linux: fix compile error with clang 3.8

(Edigrating CL 4303054)

diff --git a/Engine/Plugins/Editor/GLTFImporter/Source/GLTFImporter/Private/GLTFPackage.h b/Engine/Plugins/Editor/GLTFImporter/Source/GLTFImporter/Private/GLTFPackage.h
--- a/Engine/Plugins/Editor/GLTFImporter/Source/GLTFImporter/Private/GLTFPackage.h
+++ b/Engine/Plugins/Editor/GLTFImporter/Source/GLTFImporter/Private/GLTFPackage.h
@@ -38,7 +38,7 @@ UPackage* GetAssetPackageAndName(UObject* Parent, const FString& PreferredName,
 	}
 	else
 	{
-		const FString Suffix; // empty
+		const FString Suffix(TEXT("")); // empty
 
 		static FAssetToolsModule& AssetToolsModule = FModuleManager::LoadModuleChecked<FAssetToolsModule>("AssetTools");
 		FString FinalPackageName;
