Copied cl 6937252, 6988979, 7039079, 6845081 from //Tasks/UE4/Private-Datasmith-Deltagen

- GLTF: fixed material reimport
- GLTF: fixed failing import when image file wasn't found
- MDL: newly imported material can now be deleted without resorting to saving it first
- MDL: Displacement fixed

diff --git a/Engine/Plugins/Editor/GLTFImporter/Source/GLTFImporter/Private/GLTF/GLTFAsset.cpp b/Engine/Plugins/Editor/GLTFImporter/Source/GLTFImporter/Private/GLTF/GLTFAsset.cpp
--- a/Engine/Plugins/Editor/GLTFImporter/Source/GLTFImporter/Private/GLTF/GLTFAsset.cpp
+++ b/Engine/Plugins/Editor/GLTFImporter/Source/GLTFImporter/Private/GLTF/GLTFAsset.cpp
@@ -81,9 +81,10 @@ namespace GLTF
 
 		{
 			const FString& TexPrefix = Prefix + TEXT("_texture_");
-			int32          Counter   = 0;
-			for (FTexture& Tex : Textures)
+			for (int32 TextureIndex = 0; TextureIndex != Textures.Num(); ++TextureIndex)
 			{
+				FTexture& Tex = Textures[TextureIndex];
+
 				if (!Tex.Name.IsEmpty())
 					continue;
 
@@ -97,8 +98,13 @@ namespace GLTF
 				}
 				else
 				{
-					Tex.Name = TexPrefix + FString::FromInt(Counter++);
+					Tex.Name = TexPrefix;
 				}
+
+				// GLTF texture name has decorative purpose, not guaranteed to be unique
+				// only its index is unique. Same with glTF image or its source file's basename
+				// So always include texture index into texture's name to increase probability that names are unique
+				Tex.Name = FString::FromInt(TextureIndex) + TEXT("_") + Tex.Name;
 			}
 		}
 
diff --git a/Engine/Plugins/Editor/GLTFImporter/Source/GLTFImporter/Private/GLTF/GLTFReader.cpp b/Engine/Plugins/Editor/GLTFImporter/Source/GLTFImporter/Private/GLTF/GLTFReader.cpp
--- a/Engine/Plugins/Editor/GLTFImporter/Source/GLTFImporter/Private/GLTF/GLTFReader.cpp
+++ b/Engine/Plugins/Editor/GLTFImporter/Source/GLTFImporter/Private/GLTF/GLTFReader.cpp
@@ -151,7 +151,7 @@ namespace GLTF
 				}
 				else
 				{
-					Messages.Emplace(EMessageSeverity::Error, TEXT("Could not load file."));
+					Messages.Emplace(EMessageSeverity::Error, FString::Printf(TEXT("Could not load file: '%s'"), *FullPath));
 				}
 			}
 		}
@@ -456,11 +456,7 @@ namespace GLTF
 				Image.Format = ImageFormatFromFilename(Image.URI);
 
 				Image.FilePath = Path / Image.URI;
-				if (!FPaths::FileExists(Image.FilePath))
-				{
-					Messages.Emplace(EMessageSeverity::Error, TEXT("Cannot find image: ") + Image.FilePath);
-				}
-				else if (bInLoadImageData)
+				if (bInLoadImageData)
 				{
 					FArchive* Reader = IFileManager::Get().CreateFileReader(*Image.FilePath);
 					if (Reader)
diff --git a/Engine/Plugins/Editor/GLTFImporter/Source/GLTFImporter/Private/GLTFMapFactory.cpp b/Engine/Plugins/Editor/GLTFImporter/Source/GLTFImporter/Private/GLTFMapFactory.cpp
--- a/Engine/Plugins/Editor/GLTFImporter/Source/GLTFImporter/Private/GLTFMapFactory.cpp
+++ b/Engine/Plugins/Editor/GLTFImporter/Source/GLTFImporter/Private/GLTFMapFactory.cpp
@@ -128,40 +128,18 @@ namespace GLTF
 			ValueExpressions.Add(ValueExpression);
 		}
 
-		if (!CreateMultiTexture(Map, CoordinateIndex, MapName, MapChannels, MapChannelsCount, TextureMode, ValueExpressions))
+		FMaterialExpressionTexture* TexExpression = CreateTextureMap(Map, CoordinateIndex, MapName, TextureMode);
+		if (TexExpression)
 		{
-			// no texture present, make value connections
 			for (uint32 Index = 0; Index < MapChannelsCount; ++Index)
 			{
 				const FMapChannel& MapChannel = MapChannels[Index];
 
-				if (MapChannel.OutputExpression)
-				{
-					ValueExpressions[Index]->ConnectExpression(GetFirstInput(MapChannel.OutputExpression), 0);
-					MapChannel.OutputExpression->ConnectExpression(*MapChannel.MaterialInput, 0);
-				}
-				else
-					ValueExpressions[Index]->ConnectExpression(*MapChannel.MaterialInput, 0);
-			}
-		}
-	}
-
-	bool FPBRMapFactory::CreateMultiTexture(const GLTF::FTexture& Map, int CoordinateIndex, const TCHAR* MapName, const FMapChannel* MapChannels,
-	                                        uint32 MapChannelsCount, ETextureMode TextureMode, const FExpressionList& ValueExpressions)
-	{
-		FMaterialExpressionTexture* TexExpression = CreateTextureMap(Map, CoordinateIndex, MapName, TextureMode);
-		if (!TexExpression)
-			return false;
+				FMaterialExpressionGeneric* MultiplyExpression = CurrentMaterialElement->AddMaterialExpression<FMaterialExpressionGeneric>();
+				MultiplyExpression->SetExpressionName(TEXT("Multiply"));
 
-		for (uint32 Index = 0; Index < MapChannelsCount; ++Index)
-		{
-			const FMapChannel& MapChannel = MapChannels[Index];
-
-			FMaterialExpressionGeneric* MultiplyExpression = CurrentMaterialElement->AddMaterialExpression<FMaterialExpressionGeneric>();
-			MultiplyExpression->SetExpressionName(TEXT("Multiply"));
-
-			switch (MapChannel.Channel)
-			{
+				switch (MapChannel.Channel)
+				{
 				case EChannel::RG:
 				{
 					FMaterialExpressionFunctionCall* MakeFloat2 = CurrentMaterialElement->AddMaterialExpression<FMaterialExpressionFunctionCall>();
@@ -187,20 +165,38 @@ namespace GLTF
 					// single channel connection
 					TexExpression->ConnectExpression(*MultiplyExpression->GetInput(0), (int32)MapChannel.Channel);
 					break;
-			}
+				}
 
-			ValueExpressions[Index]->ConnectExpression(*MultiplyExpression->GetInput(1), 0);
-			if (MapChannel.OutputExpression)
+				ValueExpressions[Index]->ConnectExpression(*MultiplyExpression->GetInput(1), 0);
+				if (MapChannel.OutputExpression)
+				{
+					MultiplyExpression->ConnectExpression(GetFirstInput(MapChannel.OutputExpression), 0);
+					MapChannel.OutputExpression->ConnectExpression(*MapChannel.MaterialInput, 0);
+				}
+				else
+					MultiplyExpression->ConnectExpression(*MapChannel.MaterialInput, 0);
+			}
+		}
+		else
+		{
+			// no texture present, make value connections
+			for (uint32 Index = 0; Index < MapChannelsCount; ++Index)
 			{
-				MultiplyExpression->ConnectExpression(GetFirstInput(MapChannel.OutputExpression), 0);
-				MapChannel.OutputExpression->ConnectExpression(*MapChannel.MaterialInput, 0);
+				const FMapChannel& MapChannel = MapChannels[Index];
+
+				if (MapChannel.OutputExpression)
+				{
+					ValueExpressions[Index]->ConnectExpression(GetFirstInput(MapChannel.OutputExpression), 0);
+					MapChannel.OutputExpression->ConnectExpression(*MapChannel.MaterialInput, 0);
+				}
+				else
+					ValueExpressions[Index]->ConnectExpression(*MapChannel.MaterialInput, 0);
 			}
-			else
-				MultiplyExpression->ConnectExpression(*MapChannel.MaterialInput, 0);
 		}
-		return true;
 	}
 
+	
+
 	FMaterialExpressionTexture* FPBRMapFactory::CreateTextureMap(const GLTF::FTexture& Map, int CoordinateIndex, const TCHAR* MapName,
 	                                                             ETextureMode TextureMode)
 	{
diff --git a/Engine/Plugins/Editor/GLTFImporter/Source/GLTFImporter/Private/GLTFMapFactory.h b/Engine/Plugins/Editor/GLTFImporter/Source/GLTFImporter/Private/GLTFMapFactory.h
--- a/Engine/Plugins/Editor/GLTFImporter/Source/GLTFImporter/Private/GLTFMapFactory.h
+++ b/Engine/Plugins/Editor/GLTFImporter/Source/GLTFImporter/Private/GLTFMapFactory.h
@@ -95,9 +95,6 @@ namespace GLTF
 		FMaterialExpression* CreateMap(const GLTF::FTexture& Map, int CoordinateIndex, const ValueClass& Value, const TCHAR* MapName,
 		                               const TCHAR* ValueName, ETextureMode TextureMode, FMaterialExpressionInput& MaterialInput);
 
-		bool CreateMultiTexture(const GLTF::FTexture& Map, int CoordinateIndex, const TCHAR* MapName, const FMapChannel* MapChannels,
-		                        uint32 MapChannelsCount, ETextureMode TextureMode, const FExpressionList& ValueExpressions);
-
 	private:
 		ITextureFactory& TextureFactory;
 		UObject*         ParentPackage;
diff --git a/Engine/Plugins/Editor/GLTFImporter/Source/GLTFImporter/Private/GLTFTextureFactory.cpp b/Engine/Plugins/Editor/GLTFImporter/Source/GLTFImporter/Private/GLTFTextureFactory.cpp
--- a/Engine/Plugins/Editor/GLTFImporter/Source/GLTFImporter/Private/GLTFTextureFactory.cpp
+++ b/Engine/Plugins/Editor/GLTFImporter/Source/GLTFImporter/Private/GLTFTextureFactory.cpp
@@ -112,7 +112,10 @@ GLTF::ITextureElement* FGLTFTextureFactory::CreateTexture(const GLTF::FTexture&
 		bool bOperationCanceled = false;
 		Texture                 = static_cast<UTexture2D*>(Factory->FactoryCreateFile(UTexture2D::StaticClass(), AssetPackage, *TextureName, Flags,
                                                                       GltfTexture.Source.FilePath, nullptr, nullptr, bOperationCanceled));
-		Texture->AssetImportData->Update(GltfTexture.Source.FilePath);
+		if (Texture)
+		{
+			Texture->AssetImportData->Update(GltfTexture.Source.FilePath);
+		}
 	}
 	else if (GltfTexture.Source.Format != GLTF::FImage::EFormat::Unknown)
 	{
