Copied cl 7207638 from //Tasks/UE4/Private-Datasmith-Deltagen/...

GLTF: fixed missing meshes when source GLTF has non-unique names for meshes/nodes  UE-76474

- create datasmith mesh and actor elements assigned unique names

diff --git a/Engine/Plugins/Editor/GLTFImporter/Source/GLTFImporter/Private/GLTF/GLTFAsset.cpp b/Engine/Plugins/Editor/GLTFImporter/Source/GLTFImporter/Private/GLTF/GLTFAsset.cpp
--- a/Engine/Plugins/Editor/GLTFImporter/Source/GLTFImporter/Private/GLTF/GLTFAsset.cpp
+++ b/Engine/Plugins/Editor/GLTFImporter/Source/GLTFImporter/Private/GLTF/GLTFAsset.cpp
@@ -60,7 +60,6 @@ namespace GLTF
 		check(!Prefix.IsEmpty());
 
 		GLTF::GenerateNames(Prefix + TEXT("_material_"), Materials);
-		GLTF::GenerateNames(Prefix + TEXT("_mesh_"), Meshes);
 		GLTF::GenerateNames(Prefix + TEXT("_skin_"), Skins);
 		GLTF::GenerateNames(Prefix + TEXT("_animation_"), Animations);
 
@@ -86,9 +85,10 @@ namespace GLTF
 				FTexture& Tex = Textures[TextureIndex];
 
 				if (!Tex.Name.IsEmpty())
-					continue;
-
-				if (!Tex.Source.Name.IsEmpty())
+				{
+					// keep base part of the name as it was in gltf
+				}
+				else if (!Tex.Source.Name.IsEmpty())
 				{
 					Tex.Name = Tex.Source.Name;
 				}
@@ -108,6 +108,24 @@ namespace GLTF
 			}
 		}
 
+		{
+			const FString& MeshPrefix = Prefix + TEXT("_mesh_");
+			for (int32 Index = 0; Index != Meshes.Num(); ++Index)
+			{
+				FMesh& Mesh = Meshes[Index];
+
+				if (Mesh.Name.IsEmpty())
+				{
+					Mesh.Name = MeshPrefix;
+				}
+
+				// GLTF mesh name has decorative purpose, not guaranteed to be unique
+				// only its index is unique. Same with glTF image or its source file's basename
+				// So always include mesh index into meshes name to increase probability that names are unique
+				Mesh.Name = FString::FromInt(Index) + TEXT("_") + Mesh.Name;
+			}
+		}
+
 		for (FCamera& Camera : Cameras)
 		{
 			if (Camera.Name.IsEmpty())
@@ -128,6 +146,7 @@ namespace GLTF
 		}
 
 		GLTF::GenerateNames(Prefix + TEXT("_image_"), Images);
+		GLTF::GenerateNames(Prefix + TEXT("_image_"), Meshes);
 	}
 
 	void FAsset::GetRootNodes(TArray<int32>& NodeIndices)
