Fix CIS build - remove gltf include of RawMesh.h

diff --git a/Engine/Plugins/Editor/GLTFImporter/Source/GLTFImporter/Private/GLTFImportFactory.cpp b/Engine/Plugins/Editor/GLTFImporter/Source/GLTFImporter/Private/GLTFImportFactory.cpp
--- a/Engine/Plugins/Editor/GLTFImporter/Source/GLTFImporter/Private/GLTFImportFactory.cpp
+++ b/Engine/Plugins/Editor/GLTFImporter/Source/GLTFImporter/Private/GLTFImportFactory.cpp
@@ -9,7 +9,6 @@
 #include "Misc/FileHelper.h"
 #include "Misc/Paths.h"
 #include "Engine/StaticMesh.h"
-#include "RawMesh.h"
 #include "Serialization/ArrayReader.h"
 #include "Editor/UnrealEd/Public/Editor.h"
 #include "Serialization/JsonReader.h"
