Fix necessary because of FString constructor change in CL# 4067443.

diff --git a/Engine/Plugins/Editor/GLTFImporter/Source/GLTFImporter/Private/GLTFPackage.h b/Engine/Plugins/Editor/GLTFImporter/Source/GLTFImporter/Private/GLTFPackage.h
--- a/Engine/Plugins/Editor/GLTFImporter/Source/GLTFImporter/Private/GLTFPackage.h
+++ b/Engine/Plugins/Editor/GLTFImporter/Source/GLTFImporter/Private/GLTFPackage.h
@@ -38,7 +38,7 @@ UPackage* GetAssetPackageAndName(UObject* Parent, const FString& PreferredName,
 	}
 	else
 	{
-		const FString Suffix(0, nullptr); // empty
+		const FString Suffix; // empty
 
 		static FAssetToolsModule& AssetToolsModule = FModuleManager::LoadModuleChecked<FAssetToolsModule>("AssetTools");
 		FString FinalPackageName;
