Fix for API conversion from Main -> Editor merge

diff --git a/Engine/Plugins/Editor/GLTFImporter/Source/GLTFImporter/Private/GLTFImportFactory.cpp b/Engine/Plugins/Editor/GLTFImporter/Source/GLTFImporter/Private/GLTFImportFactory.cpp
--- a/Engine/Plugins/Editor/GLTFImporter/Source/GLTFImporter/Private/GLTFImportFactory.cpp
+++ b/Engine/Plugins/Editor/GLTFImporter/Source/GLTFImporter/Private/GLTFImportFactory.cpp
@@ -86,7 +86,7 @@ UObject* UGLTFImportFactory::FactoryCreateFile(UClass* InClass, UObject* InParen
 		}
 	}
 
-	FEditorDelegates::OnAssetPostImport.Broadcast(this, Object);
+	GEditor->GetEditorSubsystem<UImportSubsystem>()->BroadcastAssetPostImport(this, Object);
 
 	GLTFImporterImpl::ShowLogMessages(Context.GetLogMessages());
 
