Fixed deprecation warning.

diff --git a/Engine/Plugins/Editor/GLTFImporter/Source/GLTFImporter/Private/GLTFStaticMeshImporter.cpp b/Engine/Plugins/Editor/GLTFImporter/Source/GLTFImporter/Private/GLTFStaticMeshImporter.cpp
--- a/Engine/Plugins/Editor/GLTFImporter/Source/GLTFImporter/Private/GLTFStaticMeshImporter.cpp
+++ b/Engine/Plugins/Editor/GLTFImporter/Source/GLTFImporter/Private/GLTFStaticMeshImporter.cpp
@@ -100,7 +100,7 @@ namespace GLTF
 		TArray<FVector>                         VectorBuffers[VectorBufferCount];
 		TArray<FVector4>                        Vector4dBuffers[Vector4dBufferCount];
 		TArray<uint32>                          IntBuffer;
-		TArray<FMeshDescription::FContourPoint> Contours;
+		TArray<FVertexInstanceID>				CornerVertexInstanceIDs;
 
 		friend class FStaticMeshImporter;
 	};
@@ -164,6 +164,7 @@ namespace GLTF
 	    , bGenerateLightmapUVs(false)
 	    , MeshDescription(nullptr)
 	{
+		CornerVertexInstanceIDs.SetNum(3);
 	}
 
 	UStaticMesh* FStaticMeshImporterImpl::CreateMesh(const FMesh& Mesh, UObject* ParentPackage, EObjectFlags Flags)
@@ -407,7 +408,6 @@ namespace GLTF
 		}
 
 		// Now add all vertexInstances
-		FVertexInstanceID CornerVertexInstanceIDs[3];
 		FVertexID         CornerVertexIDs[3];
 		for (uint32 TriangleIndex = 0; TriangleIndex < TriCount; ++TriangleIndex)
 		{
@@ -438,35 +438,18 @@ namespace GLTF
 				CornerVertexIDs[Corner]         = VertexID;
 			}
 
-			Contours.Empty(3);
-			for (int32 Corner = 0; Corner < 3; ++Corner)
+			// Insert a polygon into the mesh
+			TArray<FEdgeID> NewEdgeIDs;
+			const FPolygonID NewPolygonID = MeshDescription->CreatePolygon(CurrentPolygonGroupID, CornerVertexInstanceIDs, &NewEdgeIDs);
+			for (const FEdgeID NewEdgeID : NewEdgeIDs)
 			{
-				FMeshDescription::FContourPoint& ContourPoint = Contours.Emplace_GetRef();
-				// Find the matching edge ID
-				uint32 CornerIndices[2];
-				CornerIndices[0] = (Corner + 0) % 3;
-				CornerIndices[1] = (Corner + 1) % 3;
-
-				FVertexID EdgeVertexIDs[2];
-				EdgeVertexIDs[0] = CornerVertexIDs[CornerIndices[0]];
-				EdgeVertexIDs[1] = CornerVertexIDs[CornerIndices[1]];
-
-				FEdgeID MatchEdgeId = MeshDescription->GetVertexPairEdge(EdgeVertexIDs[0], EdgeVertexIDs[1]);
-				if (MatchEdgeId == FEdgeID::Invalid)
-				{
-					MatchEdgeId = MeshDescription->CreateEdge(EdgeVertexIDs[0], EdgeVertexIDs[1]);
-					// Make all faces part of the same smoothing group, so Unreal will combine identical adjacent verts.
-					// (Is there a way to set auto-gen smoothing threshold? glTF spec says to generate flat normals if they're not specified.
-					//   We want to combine identical verts whether they're smooth neighbors or triangles belonging to the same flat polygon.)
-					EdgeHardnesses[MatchEdgeId]        = false;
-					EdgeCreaseSharpnesses[MatchEdgeId] = 0.0f;
-				}
-				ContourPoint.EdgeID           = MatchEdgeId;
-				ContourPoint.VertexInstanceID = CornerVertexInstanceIDs[CornerIndices[0]];
+				// Make all faces part of the same smoothing group, so Unreal will combine identical adjacent verts.
+				// (Is there a way to set auto-gen smoothing threshold? glTF spec says to generate flat normals if they're not specified.
+				//   We want to combine identical verts whether they're smooth neighbors or triangles belonging to the same flat polygon.)
+				EdgeHardnesses[NewEdgeID]        = false;
+				EdgeCreaseSharpnesses[NewEdgeID] = 0.0f;
 			}
 
-			// Insert a polygon into the mesh
-			const FPolygonID NewPolygonID = MeshDescription->CreatePolygon(CurrentPolygonGroupID, Contours);
 			// Triangulate the polygon
 			FMeshPolygon& Polygon = MeshDescription->GetPolygon(NewPolygonID);
 			MeshDescription->ComputePolygonTriangulation(NewPolygonID, Polygon.Triangles);
