Fixed crash in GLTF importer due to a bad merge coming from dev-editor.

diff --git a/Engine/Plugins/Editor/GLTFImporter/Source/GLTFImporter/Private/GLTFStaticMeshFactory.cpp b/Engine/Plugins/Editor/GLTFImporter/Source/GLTFImporter/Private/GLTFStaticMeshFactory.cpp
--- a/Engine/Plugins/Editor/GLTFImporter/Source/GLTFImporter/Private/GLTFStaticMeshFactory.cpp
+++ b/Engine/Plugins/Editor/GLTFImporter/Source/GLTFImporter/Private/GLTFStaticMeshFactory.cpp
@@ -100,7 +100,6 @@ namespace GLTF
 		TArray<FVector>                         VectorBuffers[VectorBufferCount];
 		TArray<FVector4>                        Vector4dBuffers[Vector4dBufferCount];
 		TArray<uint32>                          IntBuffer;
-		TArray<FMeshDescription::FContourPoint> Contours;
 		TArray<FVertexInstanceID>				CornerVertexInstanceIDs;
 		uint32                                  MaxReserveSize;
 
@@ -249,7 +248,7 @@ namespace GLTF
 		for (int32 Index = 0; Index < Mesh.Primitives.Num(); ++Index)
 		{
 			const FPrimitive& Primitive = Mesh.Primitives[Index];
-			const bool        basDegenerateTriangles =
+			const bool        bHasDegenerateTriangles =
 			    ImportPrimitive(Primitive, Index, NumUVs, Mesh.HasTangents(), Mesh.HasColors(),  //
 			                    VertexInstanceNormals, VertexInstanceTangents, VertexInstanceBinormalSigns, VertexInstanceUVs,
 			                    VertexInstanceColors,  //
@@ -265,10 +264,10 @@ namespace GLTF
 				}
 			}
 
-			if (basDegenerateTriangles)
+			if (bHasDegenerateTriangles)
 			{
 				Messages.Emplace(EMessageSeverity::Warning,
-				                 FString::Printf(TEXT("Mesh %s has primitive with with degenerate triangles: %d"), *Mesh.Name, Index));
+				                 FString::Printf(TEXT("Mesh %s has primitive with degenerate triangles: %d"), *Mesh.Name, Index));
 			}
 		}
 
@@ -417,7 +416,7 @@ namespace GLTF
 			}
 		}
 
-		bool basDegenerateTriangles = false;
+		bool bHasDegenerateTriangles = false;
 		// Now add all vertexInstances
 		FVertexID         CornerVertexIDs[3];
 		for (uint32 TriangleIndex = 0; TriangleIndex < TriCount; ++TriangleIndex)
@@ -430,8 +429,32 @@ namespace GLTF
 				const FVertexID          VertexID         = PositionIndexToVertexIdPerPrim[PrimitiveIndex][VertexIndex];
 				const FVertexInstanceID& VertexInstanceID = MeshDescription->CreateVertexInstance(VertexID);
 
+				CornerVertexInstanceIDs[Corner] = VertexInstanceID;
+				CornerVertexIDs[Corner]         = VertexID;
+			}
+
+			// Check for degenerate triangles
+			const FVertexID& Vertex1 = CornerVertexIDs[0];
+			const FVertexID& Vertex2 = CornerVertexIDs[1];
+			const FVertexID& Vertex3 = CornerVertexIDs[2];
+
+			if (Vertex1 == Vertex2 || Vertex2 == Vertex3 || Vertex1 == Vertex3)
+			{
+				bHasDegenerateTriangles = true;
+				continue; // Triangle is degenerate, skip it
+			}
+
+			for (int32 Corner = 0; Corner < 3; ++Corner)
+			{
+				const uint32 IndiceIndex = TriangleIndex * 3 + Corner;
+
+				const FVertexInstanceID& VertexInstanceID = CornerVertexInstanceIDs[Corner];
+
 				if (Tangents.Num() > 0)
+				{
 					VertexInstanceTangents[VertexInstanceID] = Tangents[IndiceIndex];
+				}
+
 				VertexInstanceNormals[VertexInstanceID] = Normals[IndiceIndex];
 				VertexInstanceBinormalSigns[VertexInstanceID] =
 				    GetBasisDeterminantSign(VertexInstanceTangents[VertexInstanceID].GetSafeNormal(),
@@ -442,16 +465,17 @@ namespace GLTF
 				{
 					VertexInstanceUVs.Set(VertexInstanceID, UVIndex, (*UVs[UVIndex])[IndiceIndex]);
 				}
+
 				if (Colors.Num() > 0)
+				{
 					VertexInstanceColors[VertexInstanceID] = Colors[IndiceIndex];
-
-				CornerVertexInstanceIDs[Corner] = VertexInstanceID;
-				CornerVertexIDs[Corner]         = VertexID;
+				}
 			}
 
 			// Insert a polygon into the mesh
 			TArray<FEdgeID> NewEdgeIDs;
 			const FPolygonID NewPolygonID = MeshDescription->CreatePolygon(CurrentPolygonGroupID, CornerVertexInstanceIDs, &NewEdgeIDs);
+
 			for (const FEdgeID NewEdgeID : NewEdgeIDs)
 			{
 				// Make all faces part of the same smoothing group, so Unreal will combine identical adjacent verts.
@@ -461,23 +485,11 @@ namespace GLTF
 				EdgeCreaseSharpnesses[NewEdgeID] = 0.0f;
 			}
 
-			// Check for degenerate triangles
-			const FMeshDescription::FContourPoint& ContourPoint1 = Contours[0];
-			const FMeshDescription::FContourPoint& ContourPoint2 = Contours[1];
-			const FMeshDescription::FContourPoint& ContourPoint3 = Contours[2];
-			if (ContourPoint1.EdgeID != ContourPoint2.EdgeID && ContourPoint2.EdgeID != ContourPoint3.EdgeID &&
-			    ContourPoint1.EdgeID != ContourPoint3.EdgeID)
-			{
-				// Triangulate the polygon
-				FMeshPolygon& Polygon = MeshDescription->GetPolygon(NewPolygonID);
-				MeshDescription->ComputePolygonTriangulation(NewPolygonID, Polygon.Triangles);
-			}
-			else
-			{
-				basDegenerateTriangles = true;
-			}
+			// Triangulate the polygon
+			FMeshPolygon& Polygon = MeshDescription->GetPolygon(NewPolygonID);
+			MeshDescription->ComputePolygonTriangulation(NewPolygonID, Polygon.Triangles);
 		}
-		return basDegenerateTriangles;
+		return bHasDegenerateTriangles;
 	}
 
 	inline const TArray<UStaticMesh*>& FStaticMeshFactoryImpl::CreateMeshes(const GLTF::FAsset& Asset, UObject* ParentPackage, EObjectFlags Flags,
