Copied 4.22 fixes for the MDL and GLTF importers
cl 4753660, 4753673, 4772807, 4773080, 4776806, 4821215, 4821216, 4821232, 4830162, 4830164, 4830882, 4899208

diff --git a/Engine/Plugins/Editor/GLTFImporter/Source/GLTFImporter/Private/GLTFStaticMeshFactory.cpp b/Engine/Plugins/Editor/GLTFImporter/Source/GLTFImporter/Private/GLTFStaticMeshFactory.cpp
--- a/Engine/Plugins/Editor/GLTFImporter/Source/GLTFImporter/Private/GLTFStaticMeshFactory.cpp
+++ b/Engine/Plugins/Editor/GLTFImporter/Source/GLTFImporter/Private/GLTFStaticMeshFactory.cpp
@@ -25,7 +25,8 @@ namespace GLTF
 	private:
 		UStaticMesh* CreateMesh(const GLTF::FMesh& Mesh, UObject* ParentPackage, EObjectFlags Flags);
 
-		void SetupMeshBuildSettings(int32 NumUVs, bool bMeshHasTagents, UStaticMesh& StaticMesh, FStaticMeshSourceModel& SourceModel);
+		void SetupMeshBuildSettings(int32 NumUVs, bool bMeshHasTagents, bool bMeshHasUVs, UStaticMesh& StaticMesh,
+		                            FStaticMeshSourceModel& SourceModel);
 
 		bool ImportPrimitive(const GLTF::FPrimitive&                        Primitive,  //
 		                     int32                                          PrimitiveIndex,
@@ -179,7 +180,6 @@ namespace GLTF
 		FStaticMeshSourceModel& SourceModel = StaticMesh->AddSourceModel();
 
 		const int32 NumUVs = FMath::Max(1, GetNumUVs(Mesh));
-		SetupMeshBuildSettings(NumUVs, Mesh.HasTangents(), *StaticMesh, SourceModel);
 
 		// GLTF currently only supports LODs via MSFT_lod, for now use always 0
 		const int32 LODIndex = 0;
@@ -271,10 +271,6 @@ namespace GLTF
 			}
 		}
 
-		if (bDidGenerateTexCoords)
-		{
-			Messages.Emplace(EMessageSeverity::Warning, TEXT("Mesh has primitives with no UVs generated: ") + Mesh.Name);
-		}
 		if (bMeshUsesEmptyMaterial)
 		{
 			Messages.Emplace(EMessageSeverity::Warning, TEXT("Mesh has primitives with no materials assigned: ") + Mesh.Name);
@@ -291,12 +287,13 @@ namespace GLTF
 			Settings.bRecomputeTangents  = true;
 		}
 
+		SetupMeshBuildSettings(NumUVs, Mesh.HasTangents(), Mesh.HasTexCoords(0), *StaticMesh, SourceModel);
 		StaticMesh->CommitMeshDescription(LODIndex);
 
 		return StaticMesh;
 	}
 
-	void FStaticMeshFactoryImpl::SetupMeshBuildSettings(int32 NumUVs, bool bMeshHasTagents, UStaticMesh& StaticMesh,
+	void FStaticMeshFactoryImpl::SetupMeshBuildSettings(int32 NumUVs, bool bMeshHasTagents, bool bMeshHasUVs, UStaticMesh& StaticMesh,
 	                                                    FStaticMeshSourceModel& SourceModel)
 	{
 		FMeshBuildSettings& Settings = SourceModel.BuildSettings;
@@ -308,6 +305,14 @@ namespace GLTF
 			SourceModel.BuildSettings.DstLightmapIndex     = NumUVs;
 			SourceModel.BuildSettings.bGenerateLightmapUVs = true;
 		}
+		else if (!bMeshHasUVs)
+		{
+			// Generate automatically a UV for correct lighting if mesh has none
+			StaticMesh.LightMapCoordinateIndex             = 1;
+			SourceModel.BuildSettings.SrcLightmapIndex     = 0;
+			SourceModel.BuildSettings.DstLightmapIndex     = 1;
+			SourceModel.BuildSettings.bGenerateLightmapUVs = true;
+		}
 		else
 		{
 			StaticMesh.LightMapCoordinateIndex             = NumUVs - 1;
