Copied fixes for the GLTF importer: 4995596 & 5007731 & 4958599

Fixes crash on static mesh reimport and gltf coordinates issues

diff --git a/Engine/Plugins/Editor/GLTFImporter/Source/GLTFImporter/Private/GLTFStaticMeshFactory.cpp b/Engine/Plugins/Editor/GLTFImporter/Source/GLTFImporter/Private/GLTFStaticMeshFactory.cpp
--- a/Engine/Plugins/Editor/GLTFImporter/Source/GLTFImporter/Private/GLTFStaticMeshFactory.cpp
+++ b/Engine/Plugins/Editor/GLTFImporter/Source/GLTFImporter/Private/GLTFStaticMeshFactory.cpp
@@ -19,6 +19,7 @@ namespace GLTF
 		FStaticMeshFactoryImpl();
 
 		const TArray<UStaticMesh*>& CreateMeshes(const GLTF::FAsset& Asset, UObject* ParentPackage, EObjectFlags Flags, bool bApplyPostEditChange);
+		void FillMeshDescription(const GLTF::FMesh &Mesh, FMeshDescription* MeshDescription);
 
 		void CleanUp();
 
@@ -39,7 +40,8 @@ namespace GLTF
 		                     const TVertexInstanceAttributesRef<FVector2D>& VertexInstanceUVs,
 		                     const TVertexInstanceAttributesRef<FVector4>&  VertexInstanceColors,
 		                     const TEdgeAttributesRef<bool>&                EdgeHardnesses,
-		                     const TEdgeAttributesRef<float>&               EdgeCreaseSharpnesses);
+		                     const TEdgeAttributesRef<float>&               EdgeCreaseSharpnesses,
+		                     FMeshDescription* MeshDescription);
 
 		inline TArray<FVector4>& GetVector4dBuffer(int32 Index)
 		{
@@ -95,7 +97,6 @@ namespace GLTF
 		TSet<int32>                  MaterialIndicesUsed;
 		TMap<int32, FPolygonGroupID> MaterialIndexToPolygonGroupID;
 		TArray<FIndexVertexIdMap>    PositionIndexToVertexIdPerPrim;
-		FMeshDescription*            MeshDescription;
 
 		TArray<FVector2D>                       Vector2dBuffers[MAX_MESH_TEXTURE_COORDS_MD + 1];
 		TArray<FVector>                         VectorBuffers[VectorBufferCount];
@@ -164,7 +165,6 @@ namespace GLTF
 	FStaticMeshFactoryImpl::FStaticMeshFactoryImpl()
 	    : ImportUniformScale(1.f)
 	    , bGenerateLightmapUVs(false)
-	    , MeshDescription(nullptr)
 	{
 		CornerVertexInstanceIDs.SetNum(3);
 	}
@@ -179,30 +179,62 @@ namespace GLTF
 
 		FStaticMeshSourceModel& SourceModel = StaticMesh->AddSourceModel();
 
-		const int32 NumUVs = FMath::Max(1, GetNumUVs(Mesh));
-
 		// GLTF currently only supports LODs via MSFT_lod, for now use always 0
 		const int32 LODIndex = 0;
-		MeshDescription      = StaticMesh->CreateMeshDescription(LODIndex);
+		FMeshDescription* MeshDescription = StaticMesh->CreateMeshDescription(LODIndex);
 		StaticMesh->RegisterMeshAttributes(*MeshDescription);
 
+		FillMeshDescription(Mesh, MeshDescription);
+
+		if (Mesh.HasJointWeights())
+		{
+			Messages.Emplace(EMessageSeverity::Warning, TEXT("Mesh has joint weights which are not supported: ") + Mesh.Name);
+		}
+
+		const bool bRecomputeTangents = !Mesh.HasTangents();
+		if (bRecomputeTangents)
+		{
+			FMeshBuildSettings& Settings = SourceModel.BuildSettings;
+			Settings.bRecomputeTangents  = true;
+		}
+
+		for (int32 Index = 0; Index < Mesh.Primitives.Num(); ++Index)
+		{
+			const FPrimitive& Primitive = Mesh.Primitives[Index];
+			const FName SlotName(*FString::FromInt(Primitive.MaterialIndex));
+			const int32 MeshSlot = StaticMesh->StaticMaterials.Emplace(nullptr, SlotName, SlotName);
+			StaticMesh->SectionInfoMap.Set(0, MeshSlot, FMeshSectionInfo(MeshSlot));
+		}
+
+
+		const int32 NumUVs = FMath::Max(1, GetNumUVs(Mesh)); // Duplicated from FillMeshDescription
+		SetupMeshBuildSettings(NumUVs, Mesh.HasTangents(), Mesh.HasTexCoords(0), *StaticMesh, SourceModel);
+		StaticMesh->CommitMeshDescription(LODIndex);
+
+		return StaticMesh;
+	}
+
+	void FStaticMeshFactoryImpl::FillMeshDescription(const FMesh &Mesh, FMeshDescription* MeshDescription)
+	{
+		const int32 NumUVs = FMath::Max(1, GetNumUVs(Mesh));
+
 		TVertexAttributesRef<FVector> VertexPositions =
-		    MeshDescription->VertexAttributes().GetAttributesRef<FVector>(MeshAttribute::Vertex::Position);
+			MeshDescription->VertexAttributes().GetAttributesRef<FVector>(MeshAttribute::Vertex::Position);
 		TEdgeAttributesRef<bool>  EdgeHardnesses = MeshDescription->EdgeAttributes().GetAttributesRef<bool>(MeshAttribute::Edge::IsHard);
 		TEdgeAttributesRef<float> EdgeCreaseSharpnesses =
-		    MeshDescription->EdgeAttributes().GetAttributesRef<float>(MeshAttribute::Edge::CreaseSharpness);
+			MeshDescription->EdgeAttributes().GetAttributesRef<float>(MeshAttribute::Edge::CreaseSharpness);
 		TPolygonGroupAttributesRef<FName> PolygonGroupImportedMaterialSlotNames =
-		    MeshDescription->PolygonGroupAttributes().GetAttributesRef<FName>(MeshAttribute::PolygonGroup::ImportedMaterialSlotName);
+			MeshDescription->PolygonGroupAttributes().GetAttributesRef<FName>(MeshAttribute::PolygonGroup::ImportedMaterialSlotName);
 		TVertexInstanceAttributesRef<FVector> VertexInstanceNormals =
-		    MeshDescription->VertexInstanceAttributes().GetAttributesRef<FVector>(MeshAttribute::VertexInstance::Normal);
+			MeshDescription->VertexInstanceAttributes().GetAttributesRef<FVector>(MeshAttribute::VertexInstance::Normal);
 		TVertexInstanceAttributesRef<FVector> VertexInstanceTangents =
-		    MeshDescription->VertexInstanceAttributes().GetAttributesRef<FVector>(MeshAttribute::VertexInstance::Tangent);
+			MeshDescription->VertexInstanceAttributes().GetAttributesRef<FVector>(MeshAttribute::VertexInstance::Tangent);
 		TVertexInstanceAttributesRef<float> VertexInstanceBinormalSigns =
-		    MeshDescription->VertexInstanceAttributes().GetAttributesRef<float>(MeshAttribute::VertexInstance::BinormalSign);
+			MeshDescription->VertexInstanceAttributes().GetAttributesRef<float>(MeshAttribute::VertexInstance::BinormalSign);
 		TVertexInstanceAttributesRef<FVector2D> VertexInstanceUVs =
-		    MeshDescription->VertexInstanceAttributes().GetAttributesRef<FVector2D>(MeshAttribute::VertexInstance::TextureCoordinate);
+			MeshDescription->VertexInstanceAttributes().GetAttributesRef<FVector2D>(MeshAttribute::VertexInstance::TextureCoordinate);
 		TVertexInstanceAttributesRef<FVector4> VertexInstanceColors =
-		    MeshDescription->VertexInstanceAttributes().GetAttributesRef<FVector4>(MeshAttribute::VertexInstance::Color);
+			MeshDescription->VertexInstanceAttributes().GetAttributesRef<FVector4>(MeshAttribute::VertexInstance::Color);
 		VertexInstanceUVs.SetNumIndices(NumUVs);
 
 		MaterialIndicesUsed.Empty(10);
@@ -225,10 +257,6 @@ namespace GLTF
 				VertexPositions[VertexID] = Positions[PositionIndex] * ImportUniformScale;
 				PositionIndexToVertexId.Add(PositionIndex, VertexID);
 			}
-
-			const FName SlotName(*FString::FromInt(Primitive.MaterialIndex));
-			const int32 MeshSlot = StaticMesh->StaticMaterials.Emplace(nullptr, SlotName, SlotName);
-			StaticMesh->SectionInfoMap.Set(0, MeshSlot, FMeshSectionInfo(MeshSlot));
 		}
 
 		// Add the PolygonGroup
@@ -244,15 +272,15 @@ namespace GLTF
 
 		// Add the VertexInstance
 		bool bMeshUsesEmptyMaterial = false;
-		bool bDidGenerateTexCoords  = false;
+		bool bDidGenerateTexCoords = false;
 		for (int32 Index = 0; Index < Mesh.Primitives.Num(); ++Index)
 		{
 			const FPrimitive& Primitive = Mesh.Primitives[Index];
 			const bool        bHasDegenerateTriangles =
-			    ImportPrimitive(Primitive, Index, NumUVs, Mesh.HasTangents(), Mesh.HasColors(),  //
-			                    VertexInstanceNormals, VertexInstanceTangents, VertexInstanceBinormalSigns, VertexInstanceUVs,
-			                    VertexInstanceColors,  //
-			                    EdgeHardnesses, EdgeCreaseSharpnesses);
+				ImportPrimitive(Primitive, Index, NumUVs, Mesh.HasTangents(), Mesh.HasColors(),  //
+					VertexInstanceNormals, VertexInstanceTangents, VertexInstanceBinormalSigns, VertexInstanceUVs,
+					VertexInstanceColors,  //
+					EdgeHardnesses, EdgeCreaseSharpnesses, MeshDescription);
 
 			bMeshUsesEmptyMaterial |= Primitive.MaterialIndex == INDEX_NONE;
 			for (int32 UVIndex = 0; UVIndex < NumUVs; ++UVIndex)
@@ -267,30 +295,13 @@ namespace GLTF
 			if (bHasDegenerateTriangles)
 			{
 				Messages.Emplace(EMessageSeverity::Warning,
-				                 FString::Printf(TEXT("Mesh %s has primitive with degenerate triangles: %d"), *Mesh.Name, Index));
+					FString::Printf(TEXT("Mesh %s has primitive with degenerate triangles: %d"), *Mesh.Name, Index));
 			}
 		}
-
 		if (bMeshUsesEmptyMaterial)
 		{
 			Messages.Emplace(EMessageSeverity::Warning, TEXT("Mesh has primitives with no materials assigned: ") + Mesh.Name);
 		}
-		if (Mesh.HasJointWeights())
-		{
-			Messages.Emplace(EMessageSeverity::Warning, TEXT("Mesh has joint weights which are not supported: ") + Mesh.Name);
-		}
-
-		const bool bRecomputeTangents = !Mesh.HasTangents();
-		if (bRecomputeTangents)
-		{
-			FMeshBuildSettings& Settings = SourceModel.BuildSettings;
-			Settings.bRecomputeTangents  = true;
-		}
-
-		SetupMeshBuildSettings(NumUVs, Mesh.HasTangents(), Mesh.HasTexCoords(0), *StaticMesh, SourceModel);
-		StaticMesh->CommitMeshDescription(LODIndex);
-
-		return StaticMesh;
 	}
 
 	void FStaticMeshFactoryImpl::SetupMeshBuildSettings(int32 NumUVs, bool bMeshHasTagents, bool bMeshHasUVs, UStaticMesh& StaticMesh,
@@ -342,9 +353,9 @@ namespace GLTF
 	                                             const TVertexInstanceAttributesRef<FVector2D>& VertexInstanceUVs,
 	                                             const TVertexInstanceAttributesRef<FVector4>&  VertexInstanceColors,
 	                                             const TEdgeAttributesRef<bool>&                EdgeHardnesses,
-	                                             const TEdgeAttributesRef<float>&               EdgeCreaseSharpnesses)
+	                                             const TEdgeAttributesRef<float>&               EdgeCreaseSharpnesses,
+                                                 FMeshDescription* MeshDescription)
 	{
-		check(MeshDescription);
 
 		const FPolygonGroupID CurrentPolygonGroupID = MaterialIndexToPolygonGroupID[Primitive.MaterialIndex];
 		const uint32          TriCount              = Primitive.TriangleCount();
@@ -565,6 +576,12 @@ namespace GLTF
 		return Impl->CreateMeshes(Asset, ParentPackage, Flags, bApplyPostEditChange);
 	}
 
+	void FStaticMeshFactory::FillMeshDescription(const GLTF::FMesh &Mesh, FMeshDescription* MeshDescription)
+	{
+		Impl->FillMeshDescription(Mesh, MeshDescription);
+	}
+
+
 	const TArray<UStaticMesh*>& FStaticMeshFactory::GetMeshes() const
 	{
 		return Impl->StaticMeshes;
diff --git a/Engine/Plugins/Editor/GLTFImporter/Source/GLTFImporter/Public/GLTFStaticMeshFactory.h b/Engine/Plugins/Editor/GLTFImporter/Source/GLTFImporter/Public/GLTFStaticMeshFactory.h
--- a/Engine/Plugins/Editor/GLTFImporter/Source/GLTFImporter/Public/GLTFStaticMeshFactory.h
+++ b/Engine/Plugins/Editor/GLTFImporter/Source/GLTFImporter/Public/GLTFStaticMeshFactory.h
@@ -9,10 +9,12 @@
 #include "UObject/ObjectMacros.h"
 
 class UStaticMesh;
+struct FMeshDescription;
 
 namespace GLTF
 {
 	struct FAsset;
+	struct FMesh;
 	class FStaticMeshFactoryImpl;
 
 	class GLTFIMPORTER_API FStaticMeshFactory
@@ -23,6 +25,8 @@ namespace GLTF
 
 		const TArray<UStaticMesh*>& CreateMeshes(const GLTF::FAsset& Asset, UObject* ParentPackage, EObjectFlags Flags, bool bApplyPostEditChange);
 
+		void FillMeshDescription(const GLTF::FMesh &Mesh, FMeshDescription* MeshDescription);
+
 		const TArray<UStaticMesh*>& GetMeshes() const;
 		const TArray<FLogMessage>&  GetLogMessages() const;
 
diff --git a/Engine/Plugins/Enterprise/DatasmithContent/Source/DatasmithContent/Public/DatasmithAssetImportData.h b/Engine/Plugins/Enterprise/DatasmithContent/Source/DatasmithContent/Public/DatasmithAssetImportData.h
--- a/Engine/Plugins/Enterprise/DatasmithContent/Source/DatasmithContent/Public/DatasmithAssetImportData.h
+++ b/Engine/Plugins/Enterprise/DatasmithContent/Source/DatasmithContent/Public/DatasmithAssetImportData.h
@@ -140,6 +140,16 @@ class DATASMITHCONTENT_API UDatasmithGLTFSceneImportData : public UDatasmithScen
 	FString Source;
 };
 
+UCLASS(EditInlineNew)
+class DATASMITHCONTENT_API UDatasmithStaticMeshGLTFImportData : public UDatasmithStaticMeshImportData
+{
+	GENERATED_BODY()
+
+public:
+	UPROPERTY(VisibleAnywhere, Category = InternalProperty)
+	FString SourceMeshName;
+};
+
 UCLASS(EditInlineNew)
 class DATASMITHCONTENT_API UDatasmithDeltaGenAssetImportData : public UDatasmithAssetImportData
 {
