Merging CL #7659161 from //UE4/Dev-Enterprise to Release-4.23

GLTF: Fix nodes not being bound to level sequences

diff --git a/Engine/Plugins/Editor/GLTFImporter/Source/GLTFImporter/Private/GLTF/GLTFAsset.cpp b/Engine/Plugins/Editor/GLTFImporter/Source/GLTFImporter/Private/GLTF/GLTFAsset.cpp
--- a/Engine/Plugins/Editor/GLTFImporter/Source/GLTFImporter/Private/GLTF/GLTFAsset.cpp
+++ b/Engine/Plugins/Editor/GLTFImporter/Source/GLTFImporter/Private/GLTF/GLTFAsset.cpp
@@ -68,13 +68,18 @@ namespace GLTF
 			const FString& JoinPrefix = Prefix + TEXT("_join_");
 
 			int32 Counter[2] = {0, 0};
-			for (FNode& Node : Nodes)
+			for (int32 NodeIndex = 0; NodeIndex < Nodes.Num(); ++NodeIndex)
 			{
-				if (!Node.Name.IsEmpty())
-					continue;
+				FNode& Node = Nodes[NodeIndex];
 
-				bool bIsJoint = Node.Type == FNode::EType::Joint;
-				Node.Name     = (bIsJoint ? JoinPrefix : NodePrefix) + FString::FromInt(Counter[bIsJoint]++);
+				if (Node.Name.IsEmpty())
+				{
+					bool bIsJoint = Node.Type == FNode::EType::Joint;
+					Node.Name     = (bIsJoint ? JoinPrefix : NodePrefix) + FString::FromInt(Counter[bIsJoint]++);
+				}
+
+				// Make sure node names are unique
+				Node.Name = FString::FromInt(NodeIndex) + TEXT("_") + Node.Name;
 			}
 		}
 
