GLTF: Fixed import regression after MDL functions were moved

- replaced material function(used for normalmap 'scaling') with generated code

diff --git a/Engine/Plugins/Editor/GLTFImporter/Source/GLTFImporter/Private/GLTFMapFactory.cpp b/Engine/Plugins/Editor/GLTFImporter/Source/GLTFImporter/Private/GLTFMapFactory.cpp
--- a/Engine/Plugins/Editor/GLTFImporter/Source/GLTFImporter/Private/GLTFMapFactory.cpp
+++ b/Engine/Plugins/Editor/GLTFImporter/Source/GLTFImporter/Private/GLTFMapFactory.cpp
@@ -69,13 +69,36 @@ namespace GLTF
 
 		// GLTF specifies that the following formula is used:
 		// scaledNormal = normalize((<sampled normal texture value> * 2.0 - 1.0) * vec3(<normal scale>, <normal scale>, 1.0)).
-
-		FMaterialExpressionFunctionCall* NormalExpression = CurrentMaterialElement->AddMaterialExpression<FMaterialExpressionFunctionCall>();
-		NormalExpression->SetFunctionPathName(TEXT("/DatasmithContent/Materials/MDL/AdjustNormal.AdjustNormal"));
-		TexExpression->ConnectExpression(*NormalExpression->GetInput(0), 0);
-		ScalarExpression->ConnectExpression(*NormalExpression->GetInput(1), 0);
-
-		NormalExpression->ConnectExpression(CurrentMaterialElement->GetNormal(), 0);
+		FMaterialExpressionGeneric* NormalXY = CurrentMaterialElement->AddMaterialExpression<FMaterialExpressionGeneric>();
+		NormalXY->SetExpressionName(TEXT("ComponentMask"));
+		NormalXY->SetBoolProperty(TEXT("R"), true);
+		NormalXY->SetBoolProperty(TEXT("G"), true);
+		NormalXY->SetBoolProperty(TEXT("B"), false);
+		TexExpression->ConnectExpression(*NormalXY->GetInput(0), 0);
+
+		FMaterialExpressionGeneric* NormalZ = CurrentMaterialElement->AddMaterialExpression<FMaterialExpressionGeneric>();
+		NormalZ->SetExpressionName(TEXT("ComponentMask"));
+		NormalZ->SetBoolProperty(TEXT("R"), false);
+		NormalZ->SetBoolProperty(TEXT("G"), false);
+		NormalZ->SetBoolProperty(TEXT("B"), true);
+		TexExpression->ConnectExpression(*NormalZ->GetInput(0), 0);
+
+		FMaterialExpressionGeneric* MultiplyXYExpression = CurrentMaterialElement->AddMaterialExpression<FMaterialExpressionGeneric>();
+		MultiplyXYExpression->SetExpressionName(TEXT("Multiply"));
+		NormalXY->ConnectExpression(*MultiplyXYExpression->GetInput(0), 0);
+		ScalarExpression->ConnectExpression(*MultiplyXYExpression->GetInput(1), 0);
+
+		FMaterialExpressionGeneric* ReconstructNormalVector = CurrentMaterialElement->AddMaterialExpression<FMaterialExpressionGeneric>();
+		ReconstructNormalVector->SetExpressionName(TEXT("AppendVector"));
+		MultiplyXYExpression->ConnectExpression(*ReconstructNormalVector->GetInput(0), 0);
+		NormalZ->ConnectExpression(*ReconstructNormalVector->GetInput(1), 0);
+
+		FMaterialExpressionGeneric* Normalize = CurrentMaterialElement->AddMaterialExpression<FMaterialExpressionGeneric>();
+		Normalize->SetExpressionName(TEXT("Normalize"));
+
+		ReconstructNormalVector->ConnectExpression(*Normalize->GetInput(0), 0);
+
+		Normalize->ConnectExpression(CurrentMaterialElement->GetNormal(), 0);
 	}
 
 	FMaterialExpression* FPBRMapFactory::CreateColorMap(const GLTF::FTexture& Map, int CoordinateIndex, const FVector& Color, const TCHAR* MapName,
diff --git a/Engine/Plugins/Editor/GLTFImporter/Source/GLTFImporter/Public/GLTFMaterialExpressions.h b/Engine/Plugins/Editor/GLTFImporter/Source/GLTFImporter/Public/GLTFMaterialExpressions.h
--- a/Engine/Plugins/Editor/GLTFImporter/Source/GLTFImporter/Public/GLTFMaterialExpressions.h
+++ b/Engine/Plugins/Editor/GLTFImporter/Source/GLTFImporter/Public/GLTFMaterialExpressions.h
@@ -247,9 +247,32 @@ namespace GLTF
 		virtual FMaterialExpressionInput* GetInput(int32 Index) override;
 		virtual int32                     GetInputCount() const override;
 
+		void SetBoolProperty(const TCHAR* Name, bool Value)
+		{
+			BoolProperties.Add(Name, Value);
+		}
+
+		void SetFloatProperty(const TCHAR* Name, float Value)
+		{
+			FloatProperties.Add(Name, Value);
+		}
+
+		const TMap<FString, bool>& GetBoolProperties() const
+		{
+			return BoolProperties;
+		}
+
+		const TMap<FString, float>& GetFloatProperties() const
+		{
+			return FloatProperties;
+		}
+
 	private:
 		FString                          ExpressionName;
 		TArray<FMaterialExpressionInput> Inputs;
+
+		TMap<FString, bool> BoolProperties;
+		TMap<FString, float> FloatProperties;
 	};
 
 	class GLTFIMPORTER_API FMaterialExpressionFunctionCall : public FMaterialExpression
diff --git a/Engine/Plugins/Enterprise/DatasmithGLTFImporter/Source/DatasmithGLTFTranslator/Private/DatasmithGLTFMaterialElement.cpp b/Engine/Plugins/Enterprise/DatasmithGLTFImporter/Source/DatasmithGLTFTranslator/Private/DatasmithGLTFMaterialElement.cpp
--- a/Engine/Plugins/Enterprise/DatasmithGLTFImporter/Source/DatasmithGLTFTranslator/Private/DatasmithGLTFMaterialElement.cpp
+++ b/Engine/Plugins/Enterprise/DatasmithGLTFImporter/Source/DatasmithGLTFTranslator/Private/DatasmithGLTFMaterialElement.cpp
@@ -102,6 +102,22 @@ void FDatasmithGLTFMaterialElement::CreateExpressions(TArray<IDatasmithMaterialE
 				IDatasmithMaterialExpressionGeneric*    NewExpression = MaterialElement->AddMaterialExpression<IDatasmithMaterialExpressionGeneric>();
 				NewExpression->SetExpressionName(GenericExpression.GetExpressionName());
 
+				for (const TPair<FString, bool>& NameValue : GenericExpression.GetBoolProperties())
+				{
+					TSharedPtr<IDatasmithKeyValueProperty> PropertyPtr = FDatasmithSceneFactory::CreateKeyValueProperty(*NameValue.Key);
+					PropertyPtr->SetPropertyType(EDatasmithKeyValuePropertyType::Bool);
+					PropertyPtr->SetValue(NameValue.Value ? TEXT("True") : TEXT("False"));
+					NewExpression->AddProperty(PropertyPtr);
+				}
+
+				for (const TPair<FString, float>& NameValue : GenericExpression.GetFloatProperties())
+				{
+					TSharedPtr<IDatasmithKeyValueProperty> PropertyPtr = FDatasmithSceneFactory::CreateKeyValueProperty(*NameValue.Key);
+					PropertyPtr->SetPropertyType(EDatasmithKeyValuePropertyType::Float);
+					PropertyPtr->SetValue(*FString::SanitizeFloat(NameValue.Value));
+					NewExpression->AddProperty(PropertyPtr);
+				}
+
 				MaterialExpressions.Add(NewExpression);
 			}
 			break;
