GLTFImporter: fixed crash when material/texture/mesh shared name

diff --git a/Engine/Plugins/Editor/GLTFImporter/Source/GLTFImporter/Private/GLTFImporterContext.cpp b/Engine/Plugins/Editor/GLTFImporter/Source/GLTFImporter/Private/GLTFImporterContext.cpp
--- a/Engine/Plugins/Editor/GLTFImporter/Source/GLTFImporter/Private/GLTFImporterContext.cpp
+++ b/Engine/Plugins/Editor/GLTFImporter/Source/GLTFImporter/Private/GLTFImporterContext.cpp
@@ -14,7 +14,7 @@ class FGLTFMaterialElementFactory : public GLTF::IMaterialElementFactory
 public:
 	virtual GLTF::FMaterialElement* CreateMaterial(const TCHAR* Name, UObject* ParentPackage, EObjectFlags Flags) override
 	{
-		const FString PackageName  = UPackageTools::SanitizePackageName(FPaths::Combine(ParentPackage->GetName(), Name));
+		const FString PackageName  = UPackageTools::SanitizePackageName(FPaths::Combine(ParentPackage->GetName(), TEXT("Materials"), Name));
 		UPackage*     AssetPackage = CreatePackage(nullptr, *PackageName);
 		UMaterial*    Material     = NewObject<UMaterial>(AssetPackage, UMaterial::StaticClass(), *FPaths::GetBaseFilename(PackageName), Flags);
 
diff --git a/Engine/Plugins/Editor/GLTFImporter/Source/GLTFImporter/Private/GLTFTextureFactory.cpp b/Engine/Plugins/Editor/GLTFImporter/Source/GLTFImporter/Private/GLTFTextureFactory.cpp
--- a/Engine/Plugins/Editor/GLTFImporter/Source/GLTFImporter/Private/GLTFTextureFactory.cpp
+++ b/Engine/Plugins/Editor/GLTFImporter/Source/GLTFImporter/Private/GLTFTextureFactory.cpp
@@ -102,7 +102,7 @@ GLTF::ITextureElement* FGLTFTextureFactory::CreateTexture(const GLTF::FTexture&
 	// save texture settings if texture exists
 	Factory->SuppressImportOverwriteDialog();
 
-	const FString PackageName  = UPackageTools::SanitizePackageName(FPaths::Combine(ParentPackage->GetName(), TextureName));
+	const FString PackageName  = UPackageTools::SanitizePackageName(FPaths::Combine(ParentPackage->GetName(), TEXT("Textures"), TextureName));
 	UPackage*     AssetPackage = CreatePackage(nullptr, *PackageName);
 
 	UTexture2D* Texture = nullptr;
