GLTFImporter - fixed crash when imported animation has non-unit quaternions

diff --git a/Engine/Plugins/Editor/GLTFImporter/Source/GLTFImporter/Private/GLTF/ConversionUtilities.h b/Engine/Plugins/Editor/GLTFImporter/Source/GLTFImporter/Private/GLTF/ConversionUtilities.h
--- a/Engine/Plugins/Editor/GLTFImporter/Source/GLTFImporter/Private/GLTF/ConversionUtilities.h
+++ b/Engine/Plugins/Editor/GLTFImporter/Source/GLTFImporter/Private/GLTF/ConversionUtilities.h
@@ -41,7 +41,8 @@ namespace GLTF
 		// therefore QuatUE = (-qX, -qZ, -qY, qw)
 
 		FQuat Result(-Quat.X, -Quat.Z, -Quat.Y, Quat.W);
-		check(Result.IsNormalized());
+		// Not checking if quaternion is normalized
+		// e.g. some sources use non-unit Quats for rotation tangents
 		return Result;
 	}
 
