GLTF: fixed not displaying errors spawned during LoadStaticMesh

- the issue was about degenerate triangles that were removed but no message was shown to the user

diff --git a/Engine/Plugins/Editor/GLTFImporter/Source/GLTFCore/Private/GLTFMeshFactory.cpp b/Engine/Plugins/Editor/GLTFImporter/Source/GLTFCore/Private/GLTFMeshFactory.cpp
--- a/Engine/Plugins/Editor/GLTFImporter/Source/GLTFCore/Private/GLTFMeshFactory.cpp
+++ b/Engine/Plugins/Editor/GLTFImporter/Source/GLTFCore/Private/GLTFMeshFactory.cpp
@@ -339,13 +339,8 @@ namespace GLTF
 		{
 			for (int32 Corner = 0; Corner < 3; ++Corner)
 			{
-				const uint32 IndiceIndex = TriangleIndex * 3 + Corner;
-				const int32  VertexIndex = Indices[IndiceIndex];
-
-				const FVertexID          VertexID = PositionIndexToVertexIdPerPrim[PrimitiveIndex][VertexIndex];
-				const FVertexInstanceID& VertexInstanceID = MeshDescription->CreateVertexInstance(VertexID);
-
-				CornerVertexInstanceIDs[Corner] = VertexInstanceID;
+				const int32 VertexIndex = Indices[TriangleIndex * 3 + Corner];
+				const FVertexID VertexID = PositionIndexToVertexIdPerPrim[PrimitiveIndex][VertexIndex];
 				CornerVertexIDs[Corner] = VertexID;
 			}
 
@@ -362,9 +357,10 @@ namespace GLTF
 
 			for (int32 Corner = 0; Corner < 3; ++Corner)
 			{
-				const uint32 IndiceIndex = TriangleIndex * 3 + Corner;
+				const FVertexInstanceID VertexInstanceID = MeshDescription->CreateVertexInstance(CornerVertexIDs[Corner]);
+				CornerVertexInstanceIDs[Corner] = VertexInstanceID;
 
-				const FVertexInstanceID& VertexInstanceID = CornerVertexInstanceIDs[Corner];
+				const uint32 IndiceIndex = TriangleIndex * 3 + Corner;
 
 				if (Tangents.Num() > 0)
 				{
@@ -464,4 +460,4 @@ namespace GLTF
 	{
 		Impl->CleanUp();
 	}
-} //namespace GLTF
\ No newline at end of file
+} //namespace GLTF
