Fixed build warnings.

diff --git a/Engine/Plugins/Editor/GLTFImporter/Source/GLTFCore/Private/GLTFMeshFactory.cpp b/Engine/Plugins/Editor/GLTFImporter/Source/GLTFCore/Private/GLTFMeshFactory.cpp
--- a/Engine/Plugins/Editor/GLTFImporter/Source/GLTFCore/Private/GLTFMeshFactory.cpp
+++ b/Engine/Plugins/Editor/GLTFImporter/Source/GLTFCore/Private/GLTFMeshFactory.cpp
@@ -33,7 +33,6 @@ namespace GLTF
 			const TVertexInstanceAttributesRef<FVector2D>& VertexInstanceUVs,
 			const TVertexInstanceAttributesRef<FVector4>&  VertexInstanceColors,
 			const TEdgeAttributesRef<bool>&                EdgeHardnesses,
-			const TEdgeAttributesRef<float>&               EdgeCreaseSharpnesses,
 			FMeshDescription* MeshDescription);
 
 		inline TArray<FVector4>& GetVector4dBuffer(int32 Index)
@@ -166,7 +165,6 @@ namespace GLTF
 		FStaticMeshAttributes StaticMeshAttributes(*MeshDescription);
 		TVertexAttributesRef<FVector> VertexPositions = StaticMeshAttributes.GetVertexPositions();
 		TEdgeAttributesRef<bool>  EdgeHardnesses = StaticMeshAttributes.GetEdgeHardnesses();
-		TEdgeAttributesRef<float> EdgeCreaseSharpnesses = StaticMeshAttributes.GetEdgeCreaseSharpnesses();
 		TPolygonGroupAttributesRef<FName> PolygonGroupImportedMaterialSlotNames = StaticMeshAttributes.GetPolygonGroupMaterialSlotNames();
 		TVertexInstanceAttributesRef<FVector> VertexInstanceNormals = StaticMeshAttributes.GetVertexInstanceNormals();
 		TVertexInstanceAttributesRef<FVector> VertexInstanceTangents = StaticMeshAttributes.GetVertexInstanceTangents();
@@ -218,7 +216,7 @@ namespace GLTF
 				ImportPrimitive(Primitive, Index, NumUVs, Mesh.HasTangents(), Mesh.HasColors(),  //
 					VertexInstanceNormals, VertexInstanceTangents, VertexInstanceBinormalSigns, VertexInstanceUVs,
 					VertexInstanceColors,  //
-					EdgeHardnesses, EdgeCreaseSharpnesses, MeshDescription);
+					EdgeHardnesses, MeshDescription);
 
 			bMeshUsesEmptyMaterial |= Primitive.MaterialIndex == INDEX_NONE;
 			for (int32 UVIndex = 0; UVIndex < NumUVs; ++UVIndex)
@@ -253,7 +251,6 @@ namespace GLTF
 		const TVertexInstanceAttributesRef<FVector2D>& VertexInstanceUVs,
 		const TVertexInstanceAttributesRef<FVector4>&  VertexInstanceColors,
 		const TEdgeAttributesRef<bool>&                EdgeHardnesses,
-		const TEdgeAttributesRef<float>&               EdgeCreaseSharpnesses,
 		FMeshDescription* MeshDescription)
 	{
 
@@ -394,10 +391,6 @@ namespace GLTF
 				// (Is there a way to set auto-gen smoothing threshold? glTF spec says to generate flat normals if they're not specified.
 				//   We want to combine identical verts whether they're smooth neighbors or triangles belonging to the same flat polygon.)
 				EdgeHardnesses[NewEdgeID] = false;
-				if (EdgeCreaseSharpnesses.IsValid())
-				{
-					EdgeCreaseSharpnesses[NewEdgeID] = 0.0f;
-				}
 			}
 		}
 		return bHasDegenerateTriangles;
diff --git a/Engine/Plugins/Experimental/Text3D/Source/Text3D/Private/MeshCreator.cpp b/Engine/Plugins/Experimental/Text3D/Source/Text3D/Private/MeshCreator.cpp
--- a/Engine/Plugins/Experimental/Text3D/Source/Text3D/Private/MeshCreator.cpp
+++ b/Engine/Plugins/Experimental/Text3D/Source/Text3D/Private/MeshCreator.cpp
@@ -401,11 +401,12 @@ void FMeshCreator::MirrorGroup(const EText3DGroupType TypeIn, const EText3DGroup
 
 	for (int32 TriangleIndex = 0; TriangleIndex < TrianglesInNum; TriangleIndex++)
 	{
-		const FMeshTriangle& Triangle = MeshDescription.Triangles()[FTriangleID(GroupIn.FirstTriangle + TriangleIndex)];
+		const FTriangleID TriangleID = FTriangleID(GroupIn.FirstTriangle + TriangleIndex);
+		TArrayView<const FVertexInstanceID> VertexInstanceIDs = MeshDescription.GetTriangleVertexInstances(TriangleID);
 
-		uint32 Instance0 = static_cast<uint32>(TotalVerticesNum + Triangle.GetVertexInstanceID(0).GetValue() - GroupIn.FirstVertex);
-		uint32 Instance2 = static_cast<uint32>(TotalVerticesNum + Triangle.GetVertexInstanceID(2).GetValue() - GroupIn.FirstVertex);
-		uint32 Instance1 = static_cast<uint32>(TotalVerticesNum + Triangle.GetVertexInstanceID(1).GetValue() - GroupIn.FirstVertex);
+		uint32 Instance0 = static_cast<uint32>(TotalVerticesNum + VertexInstanceIDs[0].GetValue() - GroupIn.FirstVertex);
+		uint32 Instance2 = static_cast<uint32>(TotalVerticesNum + VertexInstanceIDs[2].GetValue() - GroupIn.FirstVertex);
+		uint32 Instance1 = static_cast<uint32>(TotalVerticesNum + VertexInstanceIDs[1].GetValue() - GroupIn.FirstVertex);
 		Data->AddTriangle(Instance0, Instance2, Instance1);
 	}
 }
diff --git a/Engine/Source/Runtime/MeshDescription/Private/MeshDescription.cpp b/Engine/Source/Runtime/MeshDescription/Private/MeshDescription.cpp
--- a/Engine/Source/Runtime/MeshDescription/Private/MeshDescription.cpp
+++ b/Engine/Source/Runtime/MeshDescription/Private/MeshDescription.cpp
@@ -1755,8 +1755,6 @@ float FMeshDescription::GetPolygonCornerAngleForVertex(const FPolygonID PolygonI
 		const FVertexID ThisVertexID = GetVertexInstanceVertex(VertexInstanceIDs[ContourIndex]);
 		const FVertexID NextVertexID = GetVertexInstanceVertex(VertexInstanceIDs[NextIndex]);
 
-		TVertexAttributesConstRef<FVector> VertexPositions = VertexAttributes().GetAttributesRef<FVector>(MeshAttribute::Vertex::Position);
-
 		const FVector PrevVertexPosition = VertexPositions[PrevVertexID];
 		const FVector ThisVertexPosition = VertexPositions[ThisVertexID];
 		const FVector NextVertexPosition = VertexPositions[NextVertexID];
