LWC: GLTFCore type fixes

diff --git a/Engine/Plugins/Editor/GLTFImporter/Source/GLTFCore/Private/GLTF/GLTFAccessor.cpp b/Engine/Plugins/Editor/GLTFImporter/Source/GLTFCore/Private/GLTF/GLTFAccessor.cpp
--- a/Engine/Plugins/Editor/GLTFImporter/Source/GLTFCore/Private/GLTF/GLTFAccessor.cpp
+++ b/Engine/Plugins/Editor/GLTFImporter/Source/GLTFCore/Private/GLTF/GLTFAccessor.cpp
@@ -164,13 +164,13 @@ namespace GLTF
 
 	void FAccessor::GetVec2Array(FVector2D* Buffer) const {}
 
-	void FAccessor::GetVec3Array(FVector* Buffer) const {}
+	void FAccessor::GetVec3Array(FVector3f* Buffer) const {}
 
 	void FAccessor::GetVec4Array(FVector4* Buffer) const {}
 
-	void FAccessor::GetMat4Array(FMatrix* Buffer) const {}
+	void FAccessor::GetMat4Array(FMatrix44f* Buffer) const {}
 
-	void FAccessor::GetCoordArray(FVector* Buffer) const
+	void FAccessor::GetCoordArray(FVector3f* Buffer) const
 	{
 		GetVec3Array(Buffer);
 		for (uint32 Index = 0; Index < Count; ++Index)
@@ -376,11 +376,11 @@ namespace GLTF
 				if (ComponentType == EComponentType::F32)
 				{
 					// copy float vec3 directly from buffer
-					return *reinterpret_cast<const FVector*>(Pointer);
+					return *reinterpret_cast<const FVector3f*>(Pointer);
 				}
 				else if (Normalized)
 				{
-					return GetNormalized<FVector, 3>(ComponentType, Pointer);
+					return GetNormalized<FVector3f, 3>(ComponentType, Pointer);
 				}
 			}
 		}
@@ -502,7 +502,7 @@ namespace GLTF
 		check(false);
 	}
 
-	void FValidAccessor::GetVec3Array(FVector* Buffer) const
+	void FValidAccessor::GetVec3Array(FVector3f* Buffer) const
 	{
 		if (Type == EType::Vec3)  // strict format match, unlike GPU shader fetch
 		{
@@ -511,12 +511,12 @@ namespace GLTF
 			if (ComponentType == EComponentType::F32)
 			{
 				// copy float vec3 directly from buffer
-				memcpy(Buffer, Src, Count * sizeof(FVector));
+				memcpy(Buffer, Src, Count * sizeof(FVector3f));
 				return;
 			}
 			else if (Normalized)
 			{
-				CopyNormalized<FVector, 3>(Buffer, Src, ComponentType, Count);
+				CopyNormalized<FVector3f, 3>(Buffer, Src, ComponentType, Count);
 				return;
 			}
 		}
@@ -544,7 +544,7 @@ namespace GLTF
 		check(false);
 	}
 
-	void FValidAccessor::GetMat4Array(FMatrix* Buffer) const
+	void FValidAccessor::GetMat4Array(FMatrix44f* Buffer) const
 	{
 		if (Type == EType::Mat4 && ComponentType == EComponentType::F32)  // strict format match, unlike GPU shader fetch
 		{
diff --git a/Engine/Plugins/Editor/GLTFImporter/Source/GLTFCore/Private/GLTF/GLTFMesh.cpp b/Engine/Plugins/Editor/GLTFImporter/Source/GLTFCore/Private/GLTF/GLTFMesh.cpp
--- a/Engine/Plugins/Editor/GLTFImporter/Source/GLTFCore/Private/GLTF/GLTFMesh.cpp
+++ b/Engine/Plugins/Editor/GLTFImporter/Source/GLTFCore/Private/GLTF/GLTFMesh.cpp
@@ -36,7 +36,7 @@ namespace GLTF
 	{
 	}
 
-	void FPrimitive::GetTangents(TArray<FVector>& Buffer) const
+	void FPrimitive::GetTangents(TArray<FVector3f>& Buffer) const
 	{
 		Buffer.Reserve(Tangent.Count);
 		for (uint32 Index = 0; Index < Tangent.Count; ++Index)
diff --git a/Engine/Plugins/Editor/GLTFImporter/Source/GLTFCore/Private/GLTFMapFactory.cpp b/Engine/Plugins/Editor/GLTFImporter/Source/GLTFCore/Private/GLTFMapFactory.cpp
--- a/Engine/Plugins/Editor/GLTFImporter/Source/GLTFCore/Private/GLTFMapFactory.cpp
+++ b/Engine/Plugins/Editor/GLTFImporter/Source/GLTFCore/Private/GLTFMapFactory.cpp
@@ -35,7 +35,7 @@ namespace GLTF
 			Expression->GetScalar() = Value;
 		}
 
-		inline void SetExpresisonValue(const FVector& Color, FMaterialExpressionColor* Expression)
+		inline void SetExpresisonValue(const FVector3f& Color, FMaterialExpressionColor* Expression)
 		{
 			Expression->GetColor() = FLinearColor(Color);
 		}
@@ -103,7 +103,7 @@ namespace GLTF
 		Normalize->ConnectExpression(CurrentMaterialElement->GetNormal(), 0);
 	}
 
-	FMaterialExpression* FPBRMapFactory::CreateColorMap(const GLTF::FTexture& Map, int CoordinateIndex, const FVector& Color, const TCHAR* MapName,
+	FMaterialExpression* FPBRMapFactory::CreateColorMap(const GLTF::FTexture& Map, int CoordinateIndex, const FVector3f& Color, const TCHAR* MapName,
 	                                                    const TCHAR* ValueName, ETextureMode TextureMode, FMaterialExpressionInput& MaterialInput)
 	{
 		return CreateMap<FMaterialExpressionColor>(Map, CoordinateIndex, Color, MapName, ValueName, TextureMode, MaterialInput);
@@ -132,7 +132,7 @@ namespace GLTF
 				{
 					ValueExpression = CurrentMaterialElement->AddMaterialExpression<FMaterialExpressionColor>();
 					FMaterialExpressionColor* ColorExpression = static_cast<FMaterialExpressionColor*>(ValueExpression);
-					SetExpresisonValue(*reinterpret_cast<const FVector*>(MapChannel.VecValue), ColorExpression);
+					SetExpresisonValue(*reinterpret_cast<const FVector3f*>(MapChannel.VecValue), ColorExpression);
 					ColorExpression->SetGroupName(*GroupName);
 					break;
 				}
diff --git a/Engine/Plugins/Editor/GLTFImporter/Source/GLTFCore/Private/GLTFMapFactory.h b/Engine/Plugins/Editor/GLTFImporter/Source/GLTFCore/Private/GLTFMapFactory.h
--- a/Engine/Plugins/Editor/GLTFImporter/Source/GLTFCore/Private/GLTFMapFactory.h
+++ b/Engine/Plugins/Editor/GLTFImporter/Source/GLTFCore/Private/GLTFMapFactory.h
@@ -51,7 +51,9 @@ namespace GLTF
 
 			void SetValue(const FVector& Vec)
 			{
-				*reinterpret_cast<FVector*>(VecValue) = Vec;
+				VecValue[0] = (float)Vec.X;
+				VecValue[1] = (float)Vec.Y;
+				VecValue[2] = (float)Vec.Z;
 			}
 
 		private:
@@ -74,7 +76,7 @@ namespace GLTF
 
 		void CreateNormalMap(const GLTF::FTexture& Map, int CoordinateIndex, float NormalScale);
 
-		FMaterialExpression* CreateColorMap(const GLTF::FTexture& Map, int CoordinateIndex, const FVector& Color, const TCHAR* MapName,
+		FMaterialExpression* CreateColorMap(const GLTF::FTexture& Map, int CoordinateIndex, const FVector3f& Color, const TCHAR* MapName,
 		                                    const TCHAR* ValueName, ETextureMode TextureMode, FMaterialExpressionInput& MaterialInput);
 
 		FMaterialExpression* CreateColorMap(const GLTF::FTexture& Map, int CoordinateIndex, const FVector4& Color, const TCHAR* MapName,
diff --git a/Engine/Plugins/Editor/GLTFImporter/Source/GLTFCore/Private/GLTFMeshFactory.cpp b/Engine/Plugins/Editor/GLTFImporter/Source/GLTFCore/Private/GLTFMeshFactory.cpp
--- a/Engine/Plugins/Editor/GLTFImporter/Source/GLTFCore/Private/GLTFMeshFactory.cpp
+++ b/Engine/Plugins/Editor/GLTFImporter/Source/GLTFCore/Private/GLTFMeshFactory.cpp
@@ -43,7 +43,7 @@ namespace GLTF
 			return Vector4dBuffers[Index];
 		}
 
-		inline TArray<FVector>& GetVectorBuffer(int32 Index)
+		inline TArray<FVector3f>& GetVectorBuffer(int32 Index)
 		{
 			check(Index < (sizeof(VectorBuffers) / sizeof(VectorBuffers[0])));
 			uint32 ReserveSize = VectorBuffers[Index].Num() + VectorBuffers[Index].GetSlack();
@@ -89,7 +89,7 @@ namespace GLTF
 		TArray<FIndexVertexIdMap>    PositionIndexToVertexIdPerPrim;
 
 		TArray<FVector2D>                       Vector2dBuffers[MAX_MESH_TEXTURE_COORDS_MD + 1];
-		TArray<FVector>                         VectorBuffers[VectorBufferCount];
+		TArray<FVector3f>                       VectorBuffers[VectorBufferCount];
 		TArray<FVector4>                        Vector4dBuffers[Vector4dBufferCount];
 		TArray<uint32>                          IntBuffer;
 		TArray<FVertexInstanceID>				CornerVertexInstanceIDs;
@@ -112,7 +112,7 @@ namespace GLTF
 			}
 		}
 
-		void GenerateFlatNormals(const TArray<FVector>& Positions, const TArray<uint32>& Indices, TArray<FVector>& Normals)
+		void GenerateFlatNormals(const TArray<FVector3f>& Positions, const TArray<uint32>& Indices, TArray<FVector3f>& Normals)
 		{
 			Normals.Empty();
 
@@ -122,11 +122,11 @@ namespace GLTF
 
 			for (uint32 i = 0; i < N; i += 3)
 			{
-				const FVector& A = Positions[Indices[i]];
-				const FVector& B = Positions[Indices[i + 1]];
-				const FVector& C = Positions[Indices[i + 2]];
+				const FVector3f& A = Positions[Indices[i]];
+				const FVector3f& B = Positions[Indices[i + 1]];
+				const FVector3f& C = Positions[Indices[i + 2]];
 
-				const FVector Normal = FVector::CrossProduct(A - B, A - C).GetSafeNormal();
+				const FVector3f Normal = FVector3f::CrossProduct(A - B, A - C).GetSafeNormal();
 				// Same for each corner of the triangle.
 				Normals[i] = Normal;
 				Normals[i + 1] = Normal;
@@ -182,7 +182,7 @@ namespace GLTF
 			// Remember which primitives use which materials.
 			MaterialIndicesUsed.Add(Primitive.MaterialIndex);
 
-			TArray<FVector>& Positions = GetVectorBuffer(PositionBufferIndex);
+			TArray<FVector3f>& Positions = GetVectorBuffer(PositionBufferIndex);
 			Primitive.GetPositions(Positions);
 
 			FIndexVertexIdMap& PositionIndexToVertexId = PositionIndexToVertexIdPerPrim[Index];
@@ -260,7 +260,7 @@ namespace GLTF
 		TArray<uint32>& Indices = GetIntBuffer();
 		Primitive.GetTriangleIndices(Indices);
 
-		TArray<FVector>& Normals = GetVectorBuffer(NormalBufferIndex);
+		TArray<FVector3f>& Normals = GetVectorBuffer(NormalBufferIndex);
 		// glTF does not guarantee each primitive within a mesh has the same attributes.
 		// Fill in gaps as needed:
 		// - missing normals will be flat, based on triangle orientation
@@ -268,22 +268,22 @@ namespace GLTF
 		// - missing tangents will be (0,0,1)
 		if (Primitive.HasNormals())
 		{
-			TArray<FVector>& ReindexBuffer = GetVectorBuffer(ReindexBufferIndex);
+			TArray<FVector3f>& ReindexBuffer = GetVectorBuffer(ReindexBufferIndex);
 			Primitive.GetNormals(Normals);
 			ReIndex(Normals, Indices, ReindexBuffer);
 			Swap(Normals, ReindexBuffer);
 		}
 		else
 		{
-			TArray<FVector>& Positions = GetVectorBuffer(PositionBufferIndex);
+			TArray<FVector3f>& Positions = GetVectorBuffer(PositionBufferIndex);
 			Primitive.GetPositions(Positions);
 			GenerateFlatNormals(Positions, Indices, Normals);
 		}
 
-		TArray<FVector>& Tangents = GetVectorBuffer(TangentBufferIndex);
+		TArray<FVector3f>& Tangents = GetVectorBuffer(TangentBufferIndex);
 		if (Primitive.HasTangents())
 		{
-			TArray<FVector>& ReindexBuffer = GetVectorBuffer(ReindexBufferIndex);
+			TArray<FVector3f>& ReindexBuffer = GetVectorBuffer(ReindexBufferIndex);
 			Primitive.GetTangents(Tangents);
 			ReIndex(Tangents, Indices, ReindexBuffer);
 			Swap(Tangents, ReindexBuffer);
@@ -401,7 +401,7 @@ namespace GLTF
 		const uint32 ReserveSize = FMath::Min(uint32(IntBuffer.Num() + IntBuffer.GetSlack()), MaxReserveSize);  // cap reserved size
 		IntBuffer.Empty(ReserveSize);
 		Vector2dBuffers[0].Empty(ReserveSize);
-		for (TArray<FVector>& Array : VectorBuffers)
+		for (TArray<FVector3f>& Array : VectorBuffers)
 		{
 			Array.Empty(ReserveSize);
 		}
diff --git a/Engine/Plugins/Editor/GLTFImporter/Source/GLTFCore/Public/GLTFAccessor.h b/Engine/Plugins/Editor/GLTFImporter/Source/GLTFCore/Public/GLTFAccessor.h
--- a/Engine/Plugins/Editor/GLTFImporter/Source/GLTFCore/Public/GLTFAccessor.h
+++ b/Engine/Plugins/Editor/GLTFImporter/Source/GLTFCore/Public/GLTFAccessor.h
@@ -102,18 +102,18 @@ namespace GLTF
 		virtual void GetFloatArray(float* Buffer) const;
 		void         GetVec2Array(TArray<FVector2D>& Buffer) const;
 		virtual void GetVec2Array(FVector2D* Buffer) const;
-		void         GetVec3Array(TArray<FVector>& Buffer) const;
-		virtual void GetVec3Array(FVector* Buffer) const;
+		void         GetVec3Array(TArray<FVector3f>& Buffer) const;
+		virtual void GetVec3Array(FVector3f* Buffer) const;
 		///@note Performs axis conversion for vec3s(i.e. from glTF right-handed and Y-up to left-handed and Z-up).
-		void         GetCoordArray(TArray<FVector>& Buffer) const;
-		void         GetCoordArray(FVector* Buffer) const;
+		void         GetCoordArray(TArray<FVector3f>& Buffer) const;
+		void         GetCoordArray(FVector3f* Buffer) const;
 		void         GetVec4Array(TArray<FVector4>& Buffer) const;
 		virtual void GetVec4Array(FVector4* Buffer) const;
 		///@note Performs axis conversion for quaternion(i.e. from glTF right-handed and Y-up to left-handed and Z-up).
 		void         GetQuatArray(TArray<FVector4>& Buffer) const;
 		void         GetQuatArray(FVector4* Buffer) const;
-		void         GetMat4Array(TArray<FMatrix>& Buffer) const;
-		virtual void GetMat4Array(FMatrix* Buffer) const;
+		void         GetMat4Array(TArray<FMatrix44f>& Buffer) const;
+		virtual void GetMat4Array(FMatrix44f* Buffer) const;
 	};
 
 	struct GLTFCORE_API FValidAccessor final : FAccessor
@@ -136,9 +136,9 @@ namespace GLTF
 		void GetUnsignedIntArray(uint32* Buffer) const override;
 		void GetFloatArray(float* Buffer) const override;
 		void GetVec2Array(FVector2D* Buffer) const override;
-		void GetVec3Array(FVector* Buffer) const override;
+		void GetVec3Array(FVector3f* Buffer) const override;
 		void GetVec4Array(FVector4* Buffer) const override;
-		void GetMat4Array(FMatrix* Buffer) const override;
+		void GetMat4Array(FMatrix44f* Buffer) const override;
 
 	private:
 		const FBufferView& BufferView;
@@ -183,14 +183,14 @@ namespace GLTF
 		GetVec2Array(Buffer.GetData());
 	}
 
-	inline void FAccessor::GetVec3Array(TArray<FVector>& Buffer) const
+	inline void FAccessor::GetVec3Array(TArray<FVector3f>& Buffer) const
 	{
 		if (IsValid())
 			Buffer.SetNumUninitialized(Count, false);
 		GetVec3Array(Buffer.GetData());
 	}
 
-	inline void FAccessor::GetCoordArray(TArray<FVector>& Buffer) const
+	inline void FAccessor::GetCoordArray(TArray<FVector3f>& Buffer) const
 	{
 		if (IsValid())
 			Buffer.SetNumUninitialized(Count, false);
@@ -211,7 +211,7 @@ namespace GLTF
 		GetQuatArray(Buffer.GetData());
 	}
 
-	inline void FAccessor::GetMat4Array(TArray<FMatrix>& Buffer) const
+	inline void FAccessor::GetMat4Array(TArray<FMatrix44f>& Buffer) const
 	{
 		if (IsValid())
 			Buffer.SetNumUninitialized(Count, false);
diff --git a/Engine/Plugins/Editor/GLTFImporter/Source/GLTFCore/Public/GLTFMaterial.h b/Engine/Plugins/Editor/GLTFImporter/Source/GLTFCore/Public/GLTFMaterial.h
--- a/Engine/Plugins/Editor/GLTFImporter/Source/GLTFCore/Public/GLTFMaterial.h
+++ b/Engine/Plugins/Editor/GLTFImporter/Source/GLTFCore/Public/GLTFMaterial.h
@@ -110,7 +110,7 @@ namespace GLTF
 		FTextureMap Emissive;
 		float       NormalScale;
 		float       OcclusionStrength;
-		FVector     EmissiveFactor;
+		FVector3f   EmissiveFactor;
 
 		// material properties
 		bool       bIsDoubleSided;
diff --git a/Engine/Plugins/Editor/GLTFImporter/Source/GLTFCore/Public/GLTFMesh.h b/Engine/Plugins/Editor/GLTFImporter/Source/GLTFCore/Public/GLTFMesh.h
--- a/Engine/Plugins/Editor/GLTFImporter/Source/GLTFCore/Public/GLTFMesh.h
+++ b/Engine/Plugins/Editor/GLTFImporter/Source/GLTFCore/Public/GLTFMesh.h
@@ -61,11 +61,11 @@ namespace GLTF
 		bool IsValid() const;
 		FMD5Hash GetHash() const;
 
-		void GetPositions(TArray<FVector>& Buffer) const;
+		void GetPositions(TArray<FVector3f>& Buffer) const;
 		bool HasNormals() const;
-		void GetNormals(TArray<FVector>& Buffer) const;
+		void GetNormals(TArray<FVector3f>& Buffer) const;
 		bool HasTangents() const;
-		void GetTangents(TArray<FVector>& Buffer) const;
+		void GetTangents(TArray<FVector3f>& Buffer) const;
 		bool HasTexCoords(uint32 Index) const;
 		void GetTexCoords(uint32 Index, TArray<FVector2D>& Buffer) const;
 		void GetColors(TArray<FVector4>& Buffer) const;
@@ -140,12 +140,12 @@ namespace GLTF
 		return Color0.IsValid();
 	}
 
-	inline void FPrimitive::GetPositions(TArray<FVector>& Buffer) const
+	inline void FPrimitive::GetPositions(TArray<FVector3f>& Buffer) const
 	{
 		Position.GetCoordArray(Buffer);
 	}
 
-	inline void FPrimitive::GetNormals(TArray<FVector>& Buffer) const
+	inline void FPrimitive::GetNormals(TArray<FVector3f>& Buffer) const
 	{
 		Normal.GetCoordArray(Buffer);
 	}
diff --git a/Engine/Plugins/Enterprise/DatasmithGLTFImporter/Source/DatasmithGLTFTranslator/Private/DatasmithGLTFAnimationImporter.cpp b/Engine/Plugins/Enterprise/DatasmithGLTFImporter/Source/DatasmithGLTFTranslator/Private/DatasmithGLTFAnimationImporter.cpp
--- a/Engine/Plugins/Enterprise/DatasmithGLTFImporter/Source/DatasmithGLTFTranslator/Private/DatasmithGLTFAnimationImporter.cpp
+++ b/Engine/Plugins/Enterprise/DatasmithGLTFImporter/Source/DatasmithGLTFTranslator/Private/DatasmithGLTFAnimationImporter.cpp
@@ -131,22 +131,22 @@ namespace DatasmithGLTFImporterImpl
 
 		void ResampleTranslation(
 			GLTF::FAnimation::EInterpolation Interpolation,
-			const FVector* FrameSourceBuffer,
+			const FVector3f* FrameSourceBuffer,
 			float ScaleFactor)
 		{
-			ResampleTrack<FVector>(FrameRate,
+			ResampleTrack<FVector3f>(FrameRate,
 				FrameTimeBuffer, FrameSourceBuffer, Interpolation,
-				[ScaleFactor](const FVector& Vec) { return Vec * ScaleFactor; },
+				[ScaleFactor](const FVector3f& Vec) { return Vec * ScaleFactor; },
 				AnimationElement, EDatasmithTransformType::Translation);
 		}
 
 		void ResampleScale(
 			GLTF::FAnimation::EInterpolation Interpolation,
-			const FVector* FrameSourceBuffer)
+			const FVector3f* FrameSourceBuffer)
 		{
-			ResampleTrack<FVector>(FrameRate,
+			ResampleTrack<FVector3f>(FrameRate,
 				FrameTimeBuffer, FrameSourceBuffer, Interpolation,
-				[](const FVector& Vec) { return Vec; },
+				[](const FVector3f& Vec) { return Vec; },
 				AnimationElement, EDatasmithTransformType::Scale);
 		}
 
@@ -252,7 +252,7 @@ uint32 FDatasmithGLTFAnimationImporter::ResampleAnimationFrames(const GLTF::FAni
 		case GLTF::FAnimation::EPath::Translation:
 		{
 			FrameDataBuffer.SetNumUninitialized(Sampler.Output.Count * 3);
-			FVector* SourceData = reinterpret_cast<FVector*>(FrameDataBuffer.GetData());
+			FVector3f* SourceData = reinterpret_cast<FVector3f*>(FrameDataBuffer.GetData());
 			Sampler.Output.GetCoordArray(SourceData);
 
 			Resampler.ResampleTranslation(Sampler.Interpolation, SourceData, ScaleFactor);
@@ -263,7 +263,7 @@ uint32 FDatasmithGLTFAnimationImporter::ResampleAnimationFrames(const GLTF::FAni
 		case GLTF::FAnimation::EPath::Scale:
 		{
 			FrameDataBuffer.SetNumUninitialized(Sampler.Output.Count * 3);
-			FVector* SourceData = reinterpret_cast<FVector*>(FrameDataBuffer.GetData());
+			FVector3f* SourceData = reinterpret_cast<FVector3f*>(FrameDataBuffer.GetData());
 			Sampler.Output.GetCoordArray(SourceData);
 
 			Resampler.ResampleScale(Sampler.Interpolation, SourceData);
