glTF Importer: Add support for PBR-Next materials

diff --git a/Engine/Plugins/Editor/GLTFImporter/Source/GLTFCore/Private/GLTF/ExtensionsHandler.cpp b/Engine/Plugins/Editor/GLTFImporter/Source/GLTFCore/Private/GLTF/ExtensionsHandler.cpp
--- a/Engine/Plugins/Editor/GLTFImporter/Source/GLTFCore/Private/GLTF/ExtensionsHandler.cpp
+++ b/Engine/Plugins/Editor/GLTFImporter/Source/GLTFCore/Private/GLTF/ExtensionsHandler.cpp
@@ -82,10 +82,14 @@ namespace GLTF
 		{
 			KHR_materials_pbrSpecularGlossiness = 0,
 			KHR_materials_unlit,
+			KHR_materials_clearcoat,
+			KHR_materials_transmission,
+			KHR_materials_sheen,
 			MSFT_packing_occlusionRoughnessMetallic,
 			MSFT_packing_normalRoughnessMetallic,
 		};
-		static const TArray<FString> Extensions = {TEXT("KHR_materials_pbrSpecularGlossiness"), TEXT("KHR_materials_unlit"),
+		static const TArray<FString> Extensions = {TEXT("KHR_materials_pbrSpecularGlossiness"), TEXT("KHR_materials_unlit"), TEXT("KHR_materials_clearcoat"),
+												   TEXT("KHR_materials_transmission"), TEXT("KHR_materials_sheen"),
 		                                           TEXT("MSFT_packing_occlusionRoughnessMetallic"), TEXT("MSFT_packing_normalRoughnessMetallic")};
 
 		const FJsonObject& ExtensionsObj = *Object.GetObjectField(TEXT("extensions"));
@@ -119,6 +123,45 @@ namespace GLTF
 					Asset->ExtensionsUsed.Add(EExtension::KHR_MaterialsUnlit);
 				}
 				break;
+				case KHR_materials_clearcoat:
+				{
+					const FJsonObject& ClearCoat = ExtObj;
+
+					Material.bHasClearCoat = true;
+
+					Material.ClearCoat.ClearCoatFactor = GetScalar(ClearCoat, TEXT("clearcoatFactor"), 1.0f);
+					GLTF::SetTextureMap(ClearCoat, TEXT("clearcoatTexture"), nullptr, Asset->Textures, Material.ClearCoat.ClearCoatMap);
+
+					Material.ClearCoat.Roughness = GetScalar(ClearCoat, TEXT("clearcoatRoughnessFactor"), 0.0f);
+					GLTF::SetTextureMap(ClearCoat, TEXT("clearcoatRoughnessTexture"), nullptr, Asset->Textures, Material.ClearCoat.RoughnessMap);
+
+					GLTF::SetTextureMap(ClearCoat, TEXT("clearcoatNormalTexture"), nullptr, Asset->Textures, Material.ClearCoat.NormalMap);
+
+					Asset->ExtensionsUsed.Add(EExtension::KHR_MaterialsClearCoat);
+				}
+				break;
+				case KHR_materials_transmission:
+				{
+					const FJsonObject& Transm = ExtObj;
+
+					Material.bHasTransmission = true;
+
+					Material.Transmission.TransmissionFactor = GetScalar(Transm, TEXT("transmissionFactor"), 0.0f);
+					GLTF::SetTextureMap(Transm, TEXT("transmissionTexture"), nullptr, Asset->Textures, Material.Transmission.TransmissionMap);
+
+					Asset->ExtensionsUsed.Add(EExtension::KHR_MaterialsTransmission);
+				}
+				break;
+				case KHR_materials_sheen:
+				{
+					const FJsonObject& Sheen = ExtObj;
+
+					Material.bHasSheen = true;
+					Material.Sheen.SheenColorFactor = GetVec3(Sheen, TEXT("sheenColorFactor"));
+
+					Asset->ExtensionsUsed.Add(EExtension::KHR_MaterialsSheen);
+				}
+				break;
 				case MSFT_packing_occlusionRoughnessMetallic:
 				{
 					const FJsonObject& Packing = ExtObj;
diff --git a/Engine/Plugins/Editor/GLTFImporter/Source/GLTFCore/Private/GLTF/GLTFAsset.cpp b/Engine/Plugins/Editor/GLTFImporter/Source/GLTFCore/Private/GLTF/GLTFAsset.cpp
--- a/Engine/Plugins/Editor/GLTFImporter/Source/GLTFCore/Private/GLTF/GLTFAsset.cpp
+++ b/Engine/Plugins/Editor/GLTFImporter/Source/GLTFCore/Private/GLTF/GLTFAsset.cpp
@@ -208,6 +208,12 @@ namespace GLTF
 				return TEXT("KHR_Materials_PbrSpecularGlossiness");
 			case GLTF::EExtension::KHR_MaterialsUnlit:
 				return TEXT("KHR_Materials_Unlit");
+			case GLTF::EExtension::KHR_MaterialsClearCoat:
+				return TEXT("KHR_MaterialsClearCoat");
+			case GLTF::EExtension::KHR_MaterialsSheen:
+				return TEXT("KHR_MaterialsSheen");
+			case GLTF::EExtension::KHR_MaterialsTransmission:
+				return TEXT("KHR_MaterialsTransmission");
 			case GLTF::EExtension::KHR_TextureTransform:
 				return TEXT("KHR_Texture_Transform");
 			case GLTF::EExtension::KHR_DracoMeshCompression:
diff --git a/Engine/Plugins/Editor/GLTFImporter/Source/GLTFCore/Private/GLTFMaterialExpressions.cpp b/Engine/Plugins/Editor/GLTFImporter/Source/GLTFCore/Private/GLTFMaterialExpressions.cpp
--- a/Engine/Plugins/Editor/GLTFImporter/Source/GLTFCore/Private/GLTFMaterialExpressions.cpp
+++ b/Engine/Plugins/Editor/GLTFImporter/Source/GLTFCore/Private/GLTFMaterialExpressions.cpp
@@ -153,6 +153,9 @@ namespace GLTF
 	    , WorldDisplacement(TEXT("WorldDisplacement"))
 	    , Refraction(TEXT("Refraction"))
 	    , AmbientOcclusion(TEXT("AmbientOcclusion"))
+		, ClearCoat(TEXT("ClearCoat"))
+		, ClearCoatRoughness(TEXT("ClearCoatRoughness"))
+		, ThinTranslucentMaterialOutput(nullptr)
 	    , bIsFinal(false)
 	{
 	}
diff --git a/Engine/Plugins/Editor/GLTFImporter/Source/GLTFCore/Private/GLTFMaterialFactory.cpp b/Engine/Plugins/Editor/GLTFImporter/Source/GLTFCore/Private/GLTFMaterialFactory.cpp
--- a/Engine/Plugins/Editor/GLTFImporter/Source/GLTFCore/Private/GLTFMaterialFactory.cpp
+++ b/Engine/Plugins/Editor/GLTFImporter/Source/GLTFCore/Private/GLTFMaterialFactory.cpp
@@ -39,6 +39,9 @@ namespace GLTF
 		void HandleOcclusion(const TArray<GLTF::FTexture>& Texture, const GLTF::FMaterial& GLTFMaterial, FPBRMapFactory& MapFactory, FMaterialElement& MaterialElement);
 		void HandleEmissive(const TArray<GLTF::FTexture>& Textures, const GLTF::FMaterial& GLTFMaterial, FPBRMapFactory& MapFactory, FMaterialElement& MaterialElement);
 		void HandleNormal(const TArray<GLTF::FTexture>& Textures, const GLTF::FMaterial& GLTFMaterial, FPBRMapFactory& MapFactory, FMaterialElement& MaterialElement);
+		void HandleClearCoat(const TArray<GLTF::FTexture>& Textures, const GLTF::FMaterial& GLTFMaterial, FPBRMapFactory& MapFactory, FMaterialElement& MaterialElement);
+		void HandleTransmission(const TArray<GLTF::FTexture>& Textures, const GLTF::FMaterial& GLTFMaterial, FPBRMapFactory& MapFactory, FMaterialElement& MaterialElement);
+		void HandleSheen(const TArray<GLTF::FTexture>& Textures, const GLTF::FMaterial& GLTFMaterial, FPBRMapFactory& MapFactory, FMaterialElement& MaterialElement);
 
 
 	private:
@@ -113,6 +116,9 @@ namespace GLTF
 
 			HandleShadingModel(Asset.Textures, GLTFMaterial, MapFactory, *MaterialElement);
 			HandleOpacity(Asset.Textures, GLTFMaterial, *MaterialElement);
+			HandleClearCoat(Asset.Textures, GLTFMaterial, MapFactory, *MaterialElement);
+			HandleTransmission(Asset.Textures, GLTFMaterial, MapFactory, *MaterialElement);
+			HandleSheen(Asset.Textures, GLTFMaterial, MapFactory, *MaterialElement);
 
 			// Additional maps
 			HandleOcclusion(Asset.Textures, GLTFMaterial, MapFactory, *MaterialElement);
@@ -329,6 +335,165 @@ namespace GLTF
 								   GLTFMaterial.NormalScale);
 	}
 
+	void FMaterialFactoryImpl::HandleClearCoat(const TArray<GLTF::FTexture>& Textures, const GLTF::FMaterial& GLTFMaterial, FPBRMapFactory& MapFactory, FMaterialElement& MaterialElement)
+	{
+		if (!GLTFMaterial.bHasClearCoat || FMath::IsNearlyEqual(GLTFMaterial.ClearCoat.ClearCoatFactor, 0.0f))
+		{
+			return;
+		}
+
+		FMaterialExpressionScalar* ClearCoatFactor = MaterialElement.AddMaterialExpression<FMaterialExpressionScalar>();
+		ClearCoatFactor->GetScalar() = GLTFMaterial.ClearCoat.ClearCoatFactor;
+		ClearCoatFactor->SetName(TEXT("ClearCoatFactor"));
+
+		FMaterialExpressionScalar* ClearCoatRoughnessFactor = MaterialElement.AddMaterialExpression<FMaterialExpressionScalar>();
+		ClearCoatRoughnessFactor->GetScalar() = GLTFMaterial.ClearCoat.Roughness;
+		ClearCoatRoughnessFactor->SetName(TEXT("ClearCoatRoughnessFactor"));
+
+		FMaterialExpressionTexture* ClearCoatTexture = MapFactory.CreateTextureMap(
+			GetTexture(GLTFMaterial.ClearCoat.ClearCoatMap, Textures),
+			GLTFMaterial.ClearCoat.ClearCoatMap.TexCoord,
+			TEXT("ClearCoat"),
+			ETextureMode::Color);
+
+		FMaterialExpressionTexture* ClearCoatRoughnessTexture = MapFactory.CreateTextureMap(
+			GetTexture(GLTFMaterial.ClearCoat.RoughnessMap, Textures),
+			GLTFMaterial.ClearCoat.RoughnessMap.TexCoord,
+			TEXT("ClearCoatRoughness"),
+			ETextureMode::Color);
+
+		FMaterialExpression* ClearCoatExpr = ClearCoatFactor;
+		if (ClearCoatTexture)
+		{
+			FMaterialExpressionGeneric* MultiplyExpr = MaterialElement.AddMaterialExpression<FMaterialExpressionGeneric>();
+			MultiplyExpr->SetExpressionName(TEXT("Multiply"));
+			ClearCoatFactor->ConnectExpression(*MultiplyExpr->GetInput(0), 0);
+			ClearCoatTexture->ConnectExpression(*MultiplyExpr->GetInput(1), (int)FPBRMapFactory::EChannel::Red);
+			ClearCoatExpr = MultiplyExpr;
+		}
+
+		FMaterialExpression* ClearCoatRougnessExpr = ClearCoatRoughnessFactor;
+		if (ClearCoatRoughnessTexture)
+		{
+			FMaterialExpressionGeneric* MultiplyExpr = MaterialElement.AddMaterialExpression<FMaterialExpressionGeneric>();
+			MultiplyExpr->SetExpressionName(TEXT("Multiply"));
+			ClearCoatRoughnessFactor->ConnectExpression(*MultiplyExpr->GetInput(0), 0);
+			ClearCoatRoughnessTexture->ConnectExpression(*MultiplyExpr->GetInput(1), (int)FPBRMapFactory::EChannel::Green);
+			ClearCoatRougnessExpr = MultiplyExpr;
+		}
+
+		ClearCoatExpr->ConnectExpression(MaterialElement.GetClearCoat(), 0);
+		ClearCoatRougnessExpr->ConnectExpression(MaterialElement.GetClearCoatRoughness(), 0);
+
+		FMaterialExpressionTexture* ClearCoatNormalTexture = MapFactory.CreateTextureMap(
+			GetTexture(GLTFMaterial.ClearCoat.NormalMap, Textures),
+			GLTFMaterial.ClearCoat.NormalMap.TexCoord,
+			TEXT("ClearCoatNormal"),
+			ETextureMode::Normal);
+
+		if (ClearCoatNormalTexture)
+		{
+			FMaterialExpressionGeneric* ClearCoatNormalCustomOutput = MaterialElement.AddMaterialExpression<FMaterialExpressionGeneric>();
+			ClearCoatNormalCustomOutput->SetExpressionName(TEXT("ClearCoatNormalCustomOutput"));
+			ClearCoatNormalTexture->ConnectExpression(*ClearCoatNormalCustomOutput->GetInput(0), 0);
+		}
+	}
+
+	void FMaterialFactoryImpl::HandleTransmission(const TArray<GLTF::FTexture>& Textures, const GLTF::FMaterial& GLTFMaterial, FPBRMapFactory& MapFactory, FMaterialElement& MaterialElement)
+	{
+		if (!GLTFMaterial.bHasTransmission)
+		{
+			return;
+		}
+
+		MaterialElement.SetBlendMode(ConvertAlphaMode(GLTF::FMaterial::EAlphaMode::Blend));
+		MaterialElement.SetShadingModel(GLTF::EGLTFMaterialShadingModel::ThinTranslucent);
+		MaterialElement.SetTranslucencyLightingMode(ETranslucencyLightingMode::TLM_SurfacePerPixelLighting);
+
+		FMaterialExpressionGeneric* ThinTranslucentOutput = MaterialElement.AddMaterialExpression<FMaterialExpressionGeneric>();
+		ThinTranslucentOutput->SetExpressionName(TEXT("ThinTranslucentMaterialOutput"));
+
+		// Connect whatever base color we have to the ThinTranslucentMaterialOutput node
+		FMaterialExpression* BaseColorExpr = MaterialElement.GetBaseColor().GetExpression();
+		BaseColorExpr->ConnectExpression(*ThinTranslucentOutput->GetInput(0), MaterialElement.GetBaseColor().GetOutputIndex());
+
+		FMaterialExpressionScalar* TransmissionFactorExpr = MaterialElement.AddMaterialExpression<FMaterialExpressionScalar>();
+		TransmissionFactorExpr->GetScalar() = GLTFMaterial.Transmission.TransmissionFactor;
+		TransmissionFactorExpr->SetName(TEXT("TransmissionFactor"));
+
+		FMaterialExpression* TransmissionExpr = TransmissionFactorExpr;
+
+		FMaterialExpressionTexture* TransmissionTexture = MapFactory.CreateTextureMap(
+			GetTexture(GLTFMaterial.Transmission.TransmissionMap, Textures), 
+			GLTFMaterial.Transmission.TransmissionMap.TexCoord, 
+			TEXT("Transmission"), 
+			ETextureMode::Color);
+
+		if (TransmissionTexture)
+		{
+			FMaterialExpressionGeneric* MultiplyExpr = MaterialElement.AddMaterialExpression<FMaterialExpressionGeneric>();
+			MultiplyExpr->SetExpressionName(TEXT("Multiply"));
+			TransmissionFactorExpr->ConnectExpression(*MultiplyExpr->GetInput(0), 0);
+			TransmissionTexture->ConnectExpression(*MultiplyExpr->GetInput(1), (int)FPBRMapFactory::EChannel::Red);
+			TransmissionExpr = MultiplyExpr;
+		}
+
+		// Connect to opacity
+		if (FMaterialExpression* OpacityInputExpr = MaterialElement.GetOpacity().GetExpression())
+		{
+			const int32 OutputIndex = MaterialElement.GetOpacity().GetOutputIndex();
+			FMaterialExpressionGeneric* MultiplyExpr = MaterialElement.AddMaterialExpression<FMaterialExpressionGeneric>();
+			MultiplyExpr->SetExpressionName(TEXT("Multiply"));
+			OpacityInputExpr->ConnectExpression(*MultiplyExpr->GetInput(0), OutputIndex);
+			TransmissionExpr->ConnectExpression(*MultiplyExpr->GetInput(1), 0);
+			TransmissionExpr = MultiplyExpr;
+		}
+
+		if (TransmissionExpr == TransmissionFactorExpr)
+		{
+			TransmissionFactorExpr->GetScalar() = 1.0f - GLTFMaterial.Transmission.TransmissionFactor;
+		}
+
+		TransmissionExpr->ConnectExpression(MaterialElement.GetOpacity(), 0);
+	}
+	
+	void FMaterialFactoryImpl::HandleSheen(const TArray<GLTF::FTexture>& Textures, const GLTF::FMaterial& GLTFMaterial, FPBRMapFactory& MapFactory, FMaterialElement& MaterialElement)
+	{
+		if (!GLTFMaterial.bHasSheen)
+		{
+			return;
+		}
+
+		FMaterialExpressionFunctionCall* FuzzyShadingCall = MaterialElement.AddMaterialExpression<FMaterialExpressionFunctionCall>();
+		FuzzyShadingCall->SetFunctionPathName(TEXT("/Engine/Functions/Engine_MaterialFunctions01/Shading/FuzzyShading.FuzzyShading"));
+
+		FMaterialExpressionColor* Fuzzyness = MaterialElement.AddMaterialExpression<FMaterialExpressionColor>();
+		Fuzzyness->SetName(TEXT("Fuzzyness"));
+		Fuzzyness->GetColor() = FLinearColor(0.9f, 0.8f, 2.0f, 1.0f);
+		Fuzzyness->ConnectExpression(*FuzzyShadingCall->GetInput(2), (int)FPBRMapFactory::EChannel::Red);	// CoreDarkness
+		Fuzzyness->ConnectExpression(*FuzzyShadingCall->GetInput(4), (int)FPBRMapFactory::EChannel::Green);	// EdgeBrightness
+		Fuzzyness->ConnectExpression(*FuzzyShadingCall->GetInput(3), (int)FPBRMapFactory::EChannel::Blue);	// Power
+
+		if (FMaterialExpression* BaseColorExpr = MaterialElement.GetBaseColor().GetExpression())
+		{
+			BaseColorExpr->ConnectExpression(*FuzzyShadingCall->GetInput(0), MaterialElement.GetBaseColor().GetOutputIndex());
+
+			FMaterialExpressionGeneric* LerpExpression = MaterialElement.AddMaterialExpression<FMaterialExpressionGeneric>();
+			LerpExpression->SetExpressionName(TEXT("LinearInterpolate"));
+
+			FuzzyShadingCall->ConnectExpression(*LerpExpression->GetInput(0), 0);
+			BaseColorExpr->ConnectExpression(*LerpExpression->GetInput(1), MaterialElement.GetBaseColor().GetOutputIndex());
+			Fuzzyness->ConnectExpression(*LerpExpression->GetInput(2), (int)FPBRMapFactory::EChannel::Alpha);
+
+			LerpExpression->ConnectExpression(MaterialElement.GetBaseColor(), 0);
+		}
+
+		if (FMaterialExpression* NormalExpr = MaterialElement.GetNormal().GetExpression())
+		{
+			NormalExpr->ConnectExpression(*FuzzyShadingCall->GetInput(1), MaterialElement.GetNormal().GetOutputIndex());
+		}
+	}
+
 	FMaterialFactory::FMaterialFactory(IMaterialElementFactory* MaterialElementFactory, ITextureFactory* TextureFactory)
 	    : Impl(new FMaterialFactoryImpl(MaterialElementFactory, TextureFactory))
 	{
diff --git a/Engine/Plugins/Editor/GLTFImporter/Source/GLTFCore/Public/GLTFAsset.h b/Engine/Plugins/Editor/GLTFImporter/Source/GLTFCore/Public/GLTFAsset.h
--- a/Engine/Plugins/Editor/GLTFImporter/Source/GLTFCore/Public/GLTFAsset.h
+++ b/Engine/Plugins/Editor/GLTFImporter/Source/GLTFCore/Public/GLTFAsset.h
@@ -17,6 +17,9 @@ namespace GLTF
 	{
 		KHR_MaterialsPbrSpecularGlossiness,
 		KHR_MaterialsUnlit,
+		KHR_MaterialsClearCoat,
+		KHR_MaterialsTransmission,
+		KHR_MaterialsSheen,
 		KHR_TextureTransform,
 		KHR_DracoMeshCompression,
 		KHR_LightsPunctual,
diff --git a/Engine/Plugins/Editor/GLTFImporter/Source/GLTFCore/Public/GLTFMaterial.h b/Engine/Plugins/Editor/GLTFImporter/Source/GLTFCore/Public/GLTFMaterial.h
--- a/Engine/Plugins/Editor/GLTFImporter/Source/GLTFCore/Public/GLTFMaterial.h
+++ b/Engine/Plugins/Editor/GLTFImporter/Source/GLTFCore/Public/GLTFMaterial.h
@@ -82,6 +82,36 @@ namespace GLTF
 			{
 			}
 		};
+		struct FClearCoat
+		{
+			float ClearCoatFactor;
+			FTextureMap ClearCoatMap;
+
+			float Roughness;
+			FTextureMap RoughnessMap;
+
+			FTextureMap NormalMap;
+
+			FClearCoat()
+				: ClearCoatFactor(1.0f)
+				, Roughness(0.0f)
+			{
+			}
+		};
+		struct FTransmission
+		{
+			float TransmissionFactor;
+			FTextureMap TransmissionMap;
+
+			FTransmission()
+				: TransmissionFactor(0.0f)
+			{
+			}
+		};
+		struct FSheen
+		{
+			FVector SheenColorFactor = FVector::Zero();
+		};
 
 		struct FPacking
 		{
@@ -103,6 +133,9 @@ namespace GLTF
 		EShadingModel       ShadingModel;
 		FMetallicRoughness  MetallicRoughness;
 		FSpecularGlossiness SpecularGlossiness;
+		FClearCoat			ClearCoat;
+		FTransmission		Transmission;
+		FSheen				Sheen;
 
 		// base properties
 		FTextureMap Normal;
@@ -120,6 +153,9 @@ namespace GLTF
 		// extension properties
 		FPacking Packing;
 		bool     bIsUnlitShadingModel;
+		bool     bHasClearCoat;
+		bool     bHasSheen;
+		bool     bHasTransmission;
 
 		FMaterial(const FString& Name)
 		    : Name(Name)
@@ -132,6 +168,9 @@ namespace GLTF
 		    , AlphaMode(EAlphaMode::Opaque)
 		    , AlphaCutoff(0.5f)
 		    , bIsUnlitShadingModel(false)
+			, bHasClearCoat(false)
+			, bHasSheen(false)
+			, bHasTransmission(false)
 		{
 		}
 
diff --git a/Engine/Plugins/Editor/GLTFImporter/Source/GLTFCore/Public/GLTFMaterialExpressions.h b/Engine/Plugins/Editor/GLTFImporter/Source/GLTFCore/Public/GLTFMaterialExpressions.h
--- a/Engine/Plugins/Editor/GLTFImporter/Source/GLTFCore/Public/GLTFMaterialExpressions.h
+++ b/Engine/Plugins/Editor/GLTFImporter/Source/GLTFCore/Public/GLTFMaterialExpressions.h
@@ -22,6 +22,14 @@ namespace GLTF
 		TextureCoordinate,
 	};
 
+	enum class EGLTFMaterialShadingModel
+	{
+		DefaultLit,
+		Subsurface,
+		ClearCoat,
+		ThinTranslucent
+	};
+
 	class GLTFCORE_API ITextureElement
 	{
 	public:
@@ -314,6 +322,8 @@ namespace GLTF
 		virtual void SetBlendMode(int InBlendMode) = 0;
 		virtual bool GetTwoSided() const           = 0;
 		virtual void SetTwoSided(bool bTwoSided)   = 0;
+		virtual void SetShadingModel(EGLTFMaterialShadingModel InShadingModel) = 0;
+		virtual void SetTranslucencyLightingMode(int InLightingMode) = 0;
 
 		virtual void Finalize() = 0;
 
@@ -332,6 +342,8 @@ namespace GLTF
 		FMaterialExpressionInput& GetWorldDisplacement();
 		FMaterialExpressionInput& GetRefraction();
 		FMaterialExpressionInput& GetAmbientOcclusion();
+		FMaterialExpressionInput& GetClearCoat();
+		FMaterialExpressionInput& GetClearCoatRoughness();
 
 		int32                GetExpressionsCount() const;
 		FMaterialExpression* GetExpression(int32 Index);
@@ -352,9 +364,13 @@ namespace GLTF
 		FMaterialExpressionInput WorldDisplacement;
 		FMaterialExpressionInput Refraction;
 		FMaterialExpressionInput AmbientOcclusion;
+		FMaterialExpressionInput ClearCoat;
+		FMaterialExpressionInput ClearCoatRoughness;
 
 		TArray<FMaterialExpression*> Expressions;
 
+		FMaterialExpression* ThinTranslucentMaterialOutput;
+
 		bool bIsFinal;
 
         FMD5Hash GLTFMaterialHash;
@@ -405,6 +421,14 @@ namespace GLTF
 	{
 		return AmbientOcclusion;
 	}
+	inline FMaterialExpressionInput& FMaterialElement::GetClearCoat()
+	{
+		return ClearCoat;
+	}
+	inline FMaterialExpressionInput& FMaterialElement::GetClearCoatRoughness()
+	{
+		return ClearCoatRoughness;
+	}
 
 	template <>
 	inline FMaterialExpressionScalar* FMaterialElement::AddMaterialExpression<FMaterialExpressionScalar>()
diff --git a/Engine/Plugins/Editor/GLTFImporter/Source/GLTFImporter/Private/GLTFImporterContext.cpp b/Engine/Plugins/Editor/GLTFImporter/Source/GLTFImporter/Private/GLTFImporterContext.cpp
--- a/Engine/Plugins/Editor/GLTFImporter/Source/GLTFImporter/Private/GLTFImporterContext.cpp
+++ b/Engine/Plugins/Editor/GLTFImporter/Source/GLTFImporter/Private/GLTFImporterContext.cpp
@@ -43,7 +43,11 @@ bool FGLTFImporterContext::OpenFile(const FString& FilePath)
 
 	// check extensions supported
 	static const TArray<GLTF::EExtension> SupportedExtensions = {GLTF::EExtension::KHR_MaterialsPbrSpecularGlossiness,
-	                                                             GLTF::EExtension::KHR_MaterialsUnlit};
+	                                                             GLTF::EExtension::KHR_MaterialsUnlit,
+																 GLTF::EExtension::KHR_MaterialsClearCoat,
+																 GLTF::EExtension::KHR_MaterialsTransmission,
+																 GLTF::EExtension::KHR_MaterialsSheen};
+
 	for (GLTF::EExtension Extension : Asset.ExtensionsUsed)
 	{
 		if (SupportedExtensions.Find(Extension) == INDEX_NONE)
diff --git a/Engine/Plugins/Editor/GLTFImporter/Source/GLTFImporter/Private/GLTFMaterialElement.cpp b/Engine/Plugins/Editor/GLTFImporter/Source/GLTFImporter/Private/GLTFMaterialElement.cpp
--- a/Engine/Plugins/Editor/GLTFImporter/Source/GLTFImporter/Private/GLTFMaterialElement.cpp
+++ b/Engine/Plugins/Editor/GLTFImporter/Source/GLTFImporter/Private/GLTFMaterialElement.cpp
@@ -14,6 +14,7 @@
 #include "Materials/MaterialExpressionTextureObjectParameter.h"
 #include "Materials/MaterialExpressionTextureSampleParameter2D.h"
 #include "Materials/MaterialExpressionVectorParameter.h"
+#include "Materials/MaterialExpressionThinTranslucentMaterialOutput.h"
 
 namespace GLTFImporterImpl
 {
@@ -233,6 +234,32 @@ void FGLTFMaterialElement::SetTwoSided(bool bTwoSided)
 	Material->TwoSided = bTwoSided;
 }
 
+void FGLTFMaterialElement::SetShadingModel(GLTF::EGLTFMaterialShadingModel InShadingModel)
+{
+	EMaterialShadingModel MaterialShadingModel;
+
+	switch (InShadingModel)
+	{
+		case GLTF::EGLTFMaterialShadingModel::ClearCoat:
+			MaterialShadingModel = EMaterialShadingModel::MSM_ClearCoat;
+			break;
+		case GLTF::EGLTFMaterialShadingModel::Subsurface:
+			MaterialShadingModel = EMaterialShadingModel::MSM_Subsurface;
+			break;
+		case GLTF::EGLTFMaterialShadingModel::ThinTranslucent:
+			MaterialShadingModel = EMaterialShadingModel::MSM_ThinTranslucent;
+			break;
+		default: MaterialShadingModel = EMaterialShadingModel::MSM_DefaultLit;
+
+	};
+	Material->SetShadingModel(MaterialShadingModel);
+}
+
+void FGLTFMaterialElement::SetTranslucencyLightingMode(int InLightingMode)
+{
+	Material->TranslucencyLightingMode = static_cast<ETranslucencyLightingMode>( InLightingMode );
+}
+
 void FGLTFMaterialElement::Finalize()
 {
 	check(!bIsFinal);
@@ -249,6 +276,26 @@ void FGLTFMaterialElement::Finalize()
 	ConnectInput(Normal, MaterialExpressions, Material->Normal);
 	ConnectInput(AmbientOcclusion, MaterialExpressions, Material->AmbientOcclusion);
 
+	// Handle transmission materials (they add a special output node to the graph)
+	if (ThinTranslucentMaterialOutput)
+	{
+		GLTF::FMaterialExpressionInput* Input = ThinTranslucentMaterialOutput->GetInput(0);
+		check(Input);
+
+		if (GLTF::FMaterialExpression* BaseColorExpr = Input->GetExpression())
+		{
+			const int32 ThinTranslucentExpressionIndex = Expressions.Find(ThinTranslucentMaterialOutput);
+			check(ThinTranslucentExpressionIndex != INDEX_NONE);
+			TStrongObjectPtr<UMaterialExpression> ThinTranslucentMaterialExpression = MaterialExpressions[ThinTranslucentExpressionIndex];
+
+			const int32 BaseColorExpressionIndex = Expressions.Find(BaseColorExpr);
+			check(BaseColorExpressionIndex != INDEX_NONE);
+			TStrongObjectPtr<UMaterialExpression> BaseColorMaterialExpression = MaterialExpressions[BaseColorExpressionIndex];
+
+			BaseColorMaterialExpression->ConnectExpression(ThinTranslucentMaterialExpression->GetInput(0), Input->GetOutputIndex());
+		}
+	}
+
 	UMaterialEditingLibrary::LayoutMaterialExpressions(Material);
 
 	Material->MarkPackageDirty();
@@ -258,11 +305,11 @@ void FGLTFMaterialElement::Finalize()
 	bIsFinal = true;
 }
 
-void FGLTFMaterialElement::CreateExpressions(TArray<TStrongObjectPtr<UMaterialExpression> >& MaterialExpressions) const
+void FGLTFMaterialElement::CreateExpressions(TArray<TStrongObjectPtr<UMaterialExpression> >& MaterialExpressions)
 {
 	MaterialExpressions.Empty(Expressions.Num());
 
-	for (const GLTF::FMaterialExpression* Expression : Expressions)
+	for (GLTF::FMaterialExpression* Expression : Expressions)
 	{
 		using namespace GLTFImporterImpl;
 
@@ -296,6 +343,11 @@ void FGLTFMaterialElement::CreateExpressions(TArray<TStrongObjectPtr<UMaterialEx
 		check(MaterialExpression);
 
 		MaterialExpressions.Add(TStrongObjectPtr<UMaterialExpression>(MaterialExpression));
+
+		if (MaterialExpression->GetClass() == UMaterialExpressionThinTranslucentMaterialOutput::StaticClass())
+		{
+			ThinTranslucentMaterialOutput = Expression;
+		}
 	}
 }
 
diff --git a/Engine/Plugins/Editor/GLTFImporter/Source/GLTFImporter/Private/GLTFMaterialElement.h b/Engine/Plugins/Editor/GLTFImporter/Source/GLTFImporter/Private/GLTFMaterialElement.h
--- a/Engine/Plugins/Editor/GLTFImporter/Source/GLTFImporter/Private/GLTFMaterialElement.h
+++ b/Engine/Plugins/Editor/GLTFImporter/Source/GLTFImporter/Private/GLTFMaterialElement.h
@@ -29,12 +29,14 @@ public:
 	virtual void SetBlendMode(int InBlendMode) override;
 	virtual bool GetTwoSided() const override;
 	virtual void SetTwoSided(bool bTwoSided) override;
+	virtual void SetShadingModel(GLTF::EGLTFMaterialShadingModel InShadingModel);
+	virtual void SetTranslucencyLightingMode(int InLightingMode) override;
 	virtual void Finalize() override;
 
 	UMaterial* GetMaterial() const;
 
 private:
-	void CreateExpressions(TArray<TStrongObjectPtr<UMaterialExpression> >& MaterialExpressions) const;
+	void CreateExpressions(TArray<TStrongObjectPtr<UMaterialExpression> >& MaterialExpressions);
 
 	void ConnectInput(const GLTF::FMaterialExpressionInput&                 ExpressionInput,
 	                  const TArray<TStrongObjectPtr<UMaterialExpression> >& MaterialExpressions, FExpressionInput& MaterialInput) const;
diff --git a/Engine/Plugins/Enterprise/DatasmithGLTFImporter/Source/DatasmithGLTFTranslator/Private/DatasmithGLTFMaterialElement.cpp b/Engine/Plugins/Enterprise/DatasmithGLTFImporter/Source/DatasmithGLTFTranslator/Private/DatasmithGLTFMaterialElement.cpp
--- a/Engine/Plugins/Enterprise/DatasmithGLTFImporter/Source/DatasmithGLTFTranslator/Private/DatasmithGLTFMaterialElement.cpp
+++ b/Engine/Plugins/Enterprise/DatasmithGLTFImporter/Source/DatasmithGLTFTranslator/Private/DatasmithGLTFMaterialElement.cpp
@@ -42,6 +42,35 @@ void FDatasmithGLTFMaterialElement::SetTwoSided(bool bTwoSided)
 	MaterialElement->SetTwoSided(bTwoSided);
 }
 
+void FDatasmithGLTFMaterialElement::SetShadingModel(GLTF::EGLTFMaterialShadingModel InShadingModel)
+{
+	EDatasmithShadingModel DatasmithShadingModel;
+
+	switch (InShadingModel)
+	{
+		case GLTF::EGLTFMaterialShadingModel::ClearCoat: 
+			DatasmithShadingModel = EDatasmithShadingModel::ClearCoat; 
+			break;
+		case GLTF::EGLTFMaterialShadingModel::Subsurface:
+			DatasmithShadingModel = EDatasmithShadingModel::Subsurface;
+			break;
+		case GLTF::EGLTFMaterialShadingModel::ThinTranslucent:
+			DatasmithShadingModel = EDatasmithShadingModel::ThinTranslucent;
+			break;
+		case GLTF::EGLTFMaterialShadingModel::DefaultLit:
+			DatasmithShadingModel = EDatasmithShadingModel::DefaultLit;
+			break;
+		default: DatasmithShadingModel = EDatasmithShadingModel::DefaultLit;
+	}
+
+	MaterialElement->SetShadingModel(DatasmithShadingModel);
+}
+
+void FDatasmithGLTFMaterialElement::SetTranslucencyLightingMode(int InLightingMode)
+{
+	MaterialElement->SetTranslucencyLightingMode(InLightingMode);
+}
+
 void FDatasmithGLTFMaterialElement::Finalize()
 {
 	check(!bIsFinal);
diff --git a/Engine/Plugins/Enterprise/DatasmithGLTFImporter/Source/DatasmithGLTFTranslator/Private/DatasmithGLTFMaterialElement.h b/Engine/Plugins/Enterprise/DatasmithGLTFImporter/Source/DatasmithGLTFTranslator/Private/DatasmithGLTFMaterialElement.h
--- a/Engine/Plugins/Enterprise/DatasmithGLTFImporter/Source/DatasmithGLTFTranslator/Private/DatasmithGLTFMaterialElement.h
+++ b/Engine/Plugins/Enterprise/DatasmithGLTFImporter/Source/DatasmithGLTFTranslator/Private/DatasmithGLTFMaterialElement.h
@@ -32,6 +32,8 @@ public:
 	virtual void SetBlendMode(int InBlendMode) override;
 	virtual bool GetTwoSided() const override;
 	virtual void SetTwoSided(bool bTwoSided) override;
+	virtual void SetShadingModel(GLTF::EGLTFMaterialShadingModel InShadingModel);
+	virtual void SetTranslucencyLightingMode(int InLightingMode) override;
 	virtual void Finalize() override;
 
 	TSharedRef<IDatasmithUEPbrMaterialElement> GetMaterial() const;
diff --git a/Engine/Plugins/Enterprise/DatasmithImporter/Source/DatasmithImporter/Private/DatasmithMaterialExpressions.cpp b/Engine/Plugins/Enterprise/DatasmithImporter/Source/DatasmithImporter/Private/DatasmithMaterialExpressions.cpp
--- a/Engine/Plugins/Enterprise/DatasmithImporter/Source/DatasmithImporter/Private/DatasmithMaterialExpressions.cpp
+++ b/Engine/Plugins/Enterprise/DatasmithImporter/Source/DatasmithImporter/Private/DatasmithMaterialExpressions.cpp
@@ -3259,6 +3259,7 @@ UMaterialInterface* FDatasmithMaterialExpressions::CreateUEPbrMaterial(UPackage*
 	UnrealMaterial->TwoSided = MaterialElement->GetTwoSided();
 	UnrealMaterial->BlendMode = GetUEPbrImportBlendMode(MaterialElement, AssetsContext);
 	UnrealMaterial->OpacityMaskClipValue = MaterialElement->GetOpacityMaskClipValue();
+	UnrealMaterial->TranslucencyLightingMode = static_cast<ETranslucencyLightingMode>(MaterialElement->GetTranslucencyLightingMode());
 
 	CreateUEPbrMaterialGraph(MaterialElement, AssetsContext, UnrealMaterial);
 
diff --git a/Engine/Source/Runtime/Datasmith/DatasmithCore/Private/DatasmithMaterialElementsImpl.cpp b/Engine/Source/Runtime/Datasmith/DatasmithCore/Private/DatasmithMaterialElementsImpl.cpp
--- a/Engine/Source/Runtime/Datasmith/DatasmithCore/Private/DatasmithMaterialElementsImpl.cpp
+++ b/Engine/Source/Runtime/Datasmith/DatasmithCore/Private/DatasmithMaterialElementsImpl.cpp
@@ -207,6 +207,7 @@ FDatasmithUEPbrMaterialElementImpl::FDatasmithUEPbrMaterialElementImpl( const TC
 	, bUseMaterialAttributes( false )
 	, bMaterialFunctionOnly ( false )
 	, OpacityMaskClipValue( 0.3333f )
+	, TranslucencyLightingMode( 0 )
 	, ShadingModel( EDatasmithShadingModel::DefaultLit )
 {
 	RegisterReferenceProxy( BaseColor, "BaseColor" );
diff --git a/Engine/Source/Runtime/Datasmith/DatasmithCore/Private/DatasmithMaterialElementsImpl.h b/Engine/Source/Runtime/Datasmith/DatasmithCore/Private/DatasmithMaterialElementsImpl.h
--- a/Engine/Source/Runtime/Datasmith/DatasmithCore/Private/DatasmithMaterialElementsImpl.h
+++ b/Engine/Source/Runtime/Datasmith/DatasmithCore/Private/DatasmithMaterialElementsImpl.h
@@ -691,6 +691,9 @@ public:
 	virtual float GetOpacityMaskClipValue() const override { return OpacityMaskClipValue; }
 	virtual void SetOpacityMaskClipValue(float InClipValue) override { OpacityMaskClipValue = InClipValue; }
 
+	virtual int GetTranslucencyLightingMode() const { return TranslucencyLightingMode; }
+	virtual void SetTranslucencyLightingMode(int InMode) { TranslucencyLightingMode = InMode; }
+
 	virtual int32 GetExpressionsCount() const override { return Expressions.View().Num(); }
 	virtual IDatasmithMaterialExpression* GetExpression( int32 Index ) override;
 	virtual int32 GetExpressionIndex( const IDatasmithMaterialExpression* Expression ) const override;
@@ -727,6 +730,7 @@ protected:
 	TReflected<bool> bUseMaterialAttributes;
 	TReflected<bool> bMaterialFunctionOnly;
 	TReflected<float> OpacityMaskClipValue;
+	TReflected<int> TranslucencyLightingMode;
 
 	TReflected<FString> ParentLabel;
 	TReflected<EDatasmithShadingModel, uint8> ShadingModel;
diff --git a/Engine/Source/Runtime/Datasmith/DatasmithCore/Private/DatasmithSceneXmlReader.cpp b/Engine/Source/Runtime/Datasmith/DatasmithCore/Private/DatasmithSceneXmlReader.cpp
--- a/Engine/Source/Runtime/Datasmith/DatasmithCore/Private/DatasmithSceneXmlReader.cpp
+++ b/Engine/Source/Runtime/Datasmith/DatasmithCore/Private/DatasmithSceneXmlReader.cpp
@@ -1934,6 +1934,10 @@ void FDatasmithSceneXmlReader::ParseUEPbrMaterial(FXmlNode* InNode, TSharedPtr<
 		{
 			OutElement->SetOpacityMaskClipValue( DatasmithSceneXmlReaderImpl::ValueFromString< float >( ChildNode->GetAttribute( TEXT("value") ) ) );
 		}
+		else if (ChildNode->GetTag() == DATASMITH_TRANSLUCENCYLIGHTINGMODE)
+		{
+			OutElement->SetTranslucencyLightingMode(DatasmithSceneXmlReaderImpl::ValueFromString< int >(ChildNode->GetAttribute(TEXT("value"))));
+		}
 		else if ( ChildNode->GetTag() == DATASMITH_FUNCTIONLYVALUENAME )
 		{
 			OutElement->SetMaterialFunctionOnly( DatasmithSceneXmlReaderImpl::ValueFromString< bool >( ChildNode->GetAttribute( TEXT("enabled") ) ) );
diff --git a/Engine/Source/Runtime/Datasmith/DatasmithCore/Private/DatasmithSceneXmlWriter.cpp b/Engine/Source/Runtime/Datasmith/DatasmithCore/Private/DatasmithSceneXmlWriter.cpp
--- a/Engine/Source/Runtime/Datasmith/DatasmithCore/Private/DatasmithSceneXmlWriter.cpp
+++ b/Engine/Source/Runtime/Datasmith/DatasmithCore/Private/DatasmithSceneXmlWriter.cpp
@@ -1435,6 +1435,7 @@ void FDatasmithSceneXmlWriterImpl::WriteUEPbrMaterialElement( const TSharedRef<
 	}
 	WriteValue( Archive, Indent + 1, DATASMITH_BLENDMODE, MaterialElement->GetBlendMode() );
 	WriteValue( Archive, Indent + 1, DATASMITH_OPACITYMASKCLIPVALUE, MaterialElement->GetOpacityMaskClipValue() );
+	WriteValue( Archive, Indent + 1, DATASMITH_TRANSLUCENCYLIGHTINGMODE, MaterialElement->GetTranslucencyLightingMode() );
 
 	if ( MaterialElement->GetShadingModel() != EDatasmithShadingModel::DefaultLit && (int32)MaterialElement->GetShadingModel() < UE_ARRAY_COUNT( DatasmithShadingModelStrings ) )
 	{
diff --git a/Engine/Source/Runtime/Datasmith/DatasmithCore/Public/DatasmithDefinitions.h b/Engine/Source/Runtime/Datasmith/DatasmithCore/Public/DatasmithDefinitions.h
--- a/Engine/Source/Runtime/Datasmith/DatasmithCore/Public/DatasmithDefinitions.h
+++ b/Engine/Source/Runtime/Datasmith/DatasmithCore/Public/DatasmithDefinitions.h
@@ -620,6 +620,7 @@ static const TCHAR* KeyValuePropertyTypeStrings[] = { TEXT("String"), TEXT("Colo
 #define DATASMITH_BLENDMODE						TEXT("Blendmode")
 #define DATASMITH_OPACITYMASKCLIPVALUE			TEXT("OpacityMaskClipValue")
 #define DATASMITH_SHADINGMODEL					TEXT("ShadingModel")
+#define DATASMITH_TRANSLUCENCYLIGHTINGMODE		TEXT("TranslucencyLightingMode")
 
 #define DATASMITH_ENVILLUMINATIONMAP			TEXT("Illuminate")
 
diff --git a/Engine/Source/Runtime/Datasmith/DatasmithCore/Public/DatasmithMaterialElements.h b/Engine/Source/Runtime/Datasmith/DatasmithCore/Public/DatasmithMaterialElements.h
--- a/Engine/Source/Runtime/Datasmith/DatasmithCore/Public/DatasmithMaterialElements.h
+++ b/Engine/Source/Runtime/Datasmith/DatasmithCore/Public/DatasmithMaterialElements.h
@@ -239,6 +239,10 @@ public:
 	virtual float GetOpacityMaskClipValue() const = 0;
 	virtual void SetOpacityMaskClipValue(float InClipValue) = 0;
 
+	virtual int GetTranslucencyLightingMode() const = 0;
+	/** InMode must match the values of ETranslucencyLightingMode from EngineTypes.h */
+	virtual void SetTranslucencyLightingMode(int InMode) = 0;
+
 	virtual int32 GetExpressionsCount() const = 0;
 	virtual IDatasmithMaterialExpression* GetExpression( int32 Index ) = 0;
 	virtual int32 GetExpressionIndex( const IDatasmithMaterialExpression* Expression ) const = 0;
