LWC disabled incremental fixes

diff --git a/Engine/Plugins/Editor/GLTFImporter/Source/GLTFCore/Private/GLTF/ConversionUtilities.h b/Engine/Plugins/Editor/GLTFImporter/Source/GLTFCore/Private/GLTF/ConversionUtilities.h
--- a/Engine/Plugins/Editor/GLTFImporter/Source/GLTFCore/Private/GLTF/ConversionUtilities.h
+++ b/Engine/Plugins/Editor/GLTFImporter/Source/GLTFCore/Private/GLTF/ConversionUtilities.h
@@ -30,7 +30,8 @@ namespace GLTF
 		return Result;
 	}
 
-	inline FQuat ConvertQuat(const FQuat& Quat)
+	template<typename T>
+	UE::Math::TQuat<T> ConvertQuat(const UE::Math::TQuat<T>& Quat)
 	{
 		// glTF uses a right-handed coordinate system, with Y up.
 		// Unreal uses a left-handed coordinate system, with Z up.
@@ -40,23 +41,7 @@ namespace GLTF
 		// also, as handedness is changed rotation is inversed - hence negation
 		// therefore QuatUE = (-qX, -qZ, -qY, qw)
 
-		FQuat Result(-Quat.X, -Quat.Z, -Quat.Y, Quat.W);
-		// Not checking if quaternion is normalized
-		// e.g. some sources use non-unit Quats for rotation tangents
-		return Result;
-	}
-
-	inline FQuat4f ConvertQuat(const FQuat4f& Quat)
-	{
-		// glTF uses a right-handed coordinate system, with Y up.
-		// Unreal uses a left-handed coordinate system, with Z up.
-		// Quat = (qX, qY, qZ, qW) = (sin(angle/2) * aX, sin(angle/2) * aY, sin(angle/2) * aZ, cons(angle/2))
-		// where (aX, aY, aZ) - rotation axis, angle - rotation angle
-		// Y swapped with Z between these coordinate systems
-		// also, as handedness is changed rotation is inversed - hence negation
-		// therefore QuatUE = (-qX, -qZ, -qY, qw)
-
-		FQuat4f Result(-Quat.X, -Quat.Z, -Quat.Y, Quat.W);
+		UE::Math::TQuat<T> Result(-Quat.X, -Quat.Z, -Quat.Y, Quat.W);
 		// Not checking if quaternion is normalized
 		// e.g. some sources use non-unit Quats for rotation tangents
 		return Result;
