Datasmith GLTF importer: Add support for PBR-Next materials

diff --git a/Engine/Plugins/Editor/GLTFImporter/Source/GLTFCore/Private/GLTF/ExtensionsHandler.cpp b/Engine/Plugins/Editor/GLTFImporter/Source/GLTFCore/Private/GLTF/ExtensionsHandler.cpp
--- a/Engine/Plugins/Editor/GLTFImporter/Source/GLTFCore/Private/GLTF/ExtensionsHandler.cpp
+++ b/Engine/Plugins/Editor/GLTFImporter/Source/GLTFCore/Private/GLTF/ExtensionsHandler.cpp
@@ -114,11 +114,13 @@ namespace GLTF
 			KHR_materials_clearcoat,
 			KHR_materials_transmission,
 			KHR_materials_sheen,
+			KHR_materials_ior,
+			KHR_materials_specular,
 			MSFT_packing_occlusionRoughnessMetallic,
 			MSFT_packing_normalRoughnessMetallic,
 		};
 		static const TArray<FString> Extensions = {TEXT("KHR_materials_pbrSpecularGlossiness"), TEXT("KHR_materials_unlit"), TEXT("KHR_materials_clearcoat"),
-												   TEXT("KHR_materials_transmission"), TEXT("KHR_materials_sheen"),
+												   TEXT("KHR_materials_transmission"), TEXT("KHR_materials_sheen"), TEXT("KHR_materials_ior"), TEXT("KHR_materials_specular"),
 		                                           TEXT("MSFT_packing_occlusionRoughnessMetallic"), TEXT("MSFT_packing_normalRoughnessMetallic")};
 
 		const FJsonObject& ExtensionsObj = *Object.GetObjectField(TEXT("extensions"));
@@ -164,7 +166,7 @@ namespace GLTF
 					Material.ClearCoat.Roughness = GetScalar(ClearCoat, TEXT("clearcoatRoughnessFactor"), 0.0f);
 					GLTF::SetTextureMap(ClearCoat, TEXT("clearcoatRoughnessTexture"), nullptr, Asset->Textures, Material.ClearCoat.RoughnessMap);
 
-					GLTF::SetTextureMap(ClearCoat, TEXT("clearcoatNormalTexture"), nullptr, Asset->Textures, Material.ClearCoat.NormalMap);
+					Material.ClearCoat.NormalMapUVScale = GLTF::SetTextureMap(ClearCoat, TEXT("clearcoatNormalTexture"), TEXT("scale"), Asset->Textures, Material.ClearCoat.NormalMap);
 
 					Asset->ExtensionsUsed.Add(EExtension::KHR_MaterialsClearCoat);
 				}
@@ -191,6 +193,29 @@ namespace GLTF
 					Asset->ExtensionsUsed.Add(EExtension::KHR_MaterialsSheen);
 				}
 				break;
+				case KHR_materials_ior:
+				{
+					const FJsonObject& IOR = ExtObj;
+
+					Material.bHasIOR = true;
+					Material.IOR = GetScalar(IOR, TEXT("ior"), 1.0f);
+
+					Asset->ExtensionsUsed.Add(EExtension::KHR_MaterialsIOR);
+				}
+				break;
+				case KHR_materials_specular:
+				{
+					const FJsonObject& Specular = ExtObj;
+
+					Material.bHasSpecular = true;
+					Material.Specular.SpecularFactor = GetScalar(Specular, TEXT("specularFactor"), 1.0f);
+					Material.Specular.SpecularColorFactor = GetVec3(Specular, TEXT("specularColorFactor"));
+					GLTF::SetTextureMap(Specular, TEXT("specularTexture"), nullptr, Asset->Textures, Material.Specular.SpecularMap);
+					GLTF::SetTextureMap(Specular, TEXT("specularColorTexture"), nullptr, Asset->Textures, Material.Specular.SpecularColorMap);
+
+					Asset->ExtensionsUsed.Add(EExtension::KHR_MaterialsSpecular);
+				}
+				break;
 				case MSFT_packing_occlusionRoughnessMetallic:
 				{
 					const FJsonObject& Packing = ExtObj;
diff --git a/Engine/Plugins/Editor/GLTFImporter/Source/GLTFCore/Private/GLTFMaterialExpressions.cpp b/Engine/Plugins/Editor/GLTFImporter/Source/GLTFCore/Private/GLTFMaterialExpressions.cpp
--- a/Engine/Plugins/Editor/GLTFImporter/Source/GLTFCore/Private/GLTFMaterialExpressions.cpp
+++ b/Engine/Plugins/Editor/GLTFImporter/Source/GLTFCore/Private/GLTFMaterialExpressions.cpp
@@ -156,6 +156,7 @@ namespace GLTF
 		, ClearCoat(TEXT("ClearCoat"))
 		, ClearCoatRoughness(TEXT("ClearCoatRoughness"))
 		, ThinTranslucentMaterialOutput(nullptr)
+		, ClearCoatBottomNormalOutput(nullptr)
 	    , bIsFinal(false)
 	{
 	}
diff --git a/Engine/Plugins/Editor/GLTFImporter/Source/GLTFCore/Private/GLTFMaterialFactory.cpp b/Engine/Plugins/Editor/GLTFImporter/Source/GLTFCore/Private/GLTFMaterialFactory.cpp
--- a/Engine/Plugins/Editor/GLTFImporter/Source/GLTFCore/Private/GLTFMaterialFactory.cpp
+++ b/Engine/Plugins/Editor/GLTFImporter/Source/GLTFCore/Private/GLTFMaterialFactory.cpp
@@ -42,6 +42,8 @@ namespace GLTF
 		void HandleClearCoat(const TArray<GLTF::FTexture>& Textures, const GLTF::FMaterial& GLTFMaterial, FPBRMapFactory& MapFactory, FMaterialElement& MaterialElement);
 		void HandleTransmission(const TArray<GLTF::FTexture>& Textures, const GLTF::FMaterial& GLTFMaterial, FPBRMapFactory& MapFactory, FMaterialElement& MaterialElement);
 		void HandleSheen(const TArray<GLTF::FTexture>& Textures, const GLTF::FMaterial& GLTFMaterial, FPBRMapFactory& MapFactory, FMaterialElement& MaterialElement);
+		void HandleIOR(const TArray<GLTF::FTexture>& Textures, const GLTF::FMaterial& GLTFMaterial, FPBRMapFactory& MapFactory, FMaterialElement& MaterialElement);
+		void HandleSpecular(const TArray<GLTF::FTexture>& Textures, const GLTF::FMaterial& GLTFMaterial, FPBRMapFactory& MapFactory, FMaterialElement& MaterialElement);
 
 
 	private:
@@ -119,6 +121,8 @@ namespace GLTF
 			HandleClearCoat(Asset.Textures, GLTFMaterial, MapFactory, *MaterialElement);
 			HandleTransmission(Asset.Textures, GLTFMaterial, MapFactory, *MaterialElement);
 			HandleSheen(Asset.Textures, GLTFMaterial, MapFactory, *MaterialElement);
+			HandleIOR(Asset.Textures, GLTFMaterial, MapFactory, *MaterialElement);
+			HandleSpecular(Asset.Textures, GLTFMaterial, MapFactory, *MaterialElement);
 
 			// Additional maps
 			HandleOcclusion(Asset.Textures, GLTFMaterial, MapFactory, *MaterialElement);
@@ -135,7 +139,7 @@ namespace GLTF
 
 	void FMaterialFactoryImpl::HandleOpacity(const TArray<GLTF::FTexture>& Textures, const GLTF::FMaterial& GLTFMaterial, FMaterialElement& MaterialElement)
 	{
-		if (GLTFMaterial.IsOpaque())
+		if (GLTFMaterial.IsOpaque() || GLTFMaterial.bHasTransmission) // With transmission, we handle opacity different
 		{
 			return;
 		}
@@ -145,56 +149,58 @@ namespace GLTF
 		FMaterialExpressionTexture* BaseColorMap = FindExpression<FMaterialExpressionTexture>(TEXT("BaseColor Map"), MaterialElement);
 		switch (GLTFMaterial.AlphaMode)
 		{
-		case FMaterial::EAlphaMode::Mask:
-		{
-			FMaterialExpressionColor* BaseColorFactor = FindExpression<FMaterialExpressionColor>(TEXT("BaseColor"), MaterialElement);
+			case FMaterial::EAlphaMode::Mask:
+			{
+				FMaterialExpressionColor* BaseColorFactor = FindExpression<FMaterialExpressionColor>(TEXT("BaseColor"), MaterialElement);
 
-			FMaterialExpressionGeneric* MultiplyExpression = MaterialElement.AddMaterialExpression<FMaterialExpressionGeneric>();
-			MultiplyExpression->SetExpressionName(TEXT("Multiply"));
-			BaseColorFactor->ConnectExpression(*MultiplyExpression->GetInput(1), (int)FPBRMapFactory::EChannel::Alpha);
-			BaseColorMap->ConnectExpression(*MultiplyExpression->GetInput(0), (int)FPBRMapFactory::EChannel::Alpha);
+				FMaterialExpressionGeneric* MultiplyExpression = MaterialElement.AddMaterialExpression<FMaterialExpressionGeneric>();
+				MultiplyExpression->SetExpressionName(TEXT("Multiply"));
 
-			FMaterialExpressionFunctionCall* CuttofExpression = MaterialElement.AddMaterialExpression<FMaterialExpressionFunctionCall>();
-			CuttofExpression->SetFunctionPathName(TEXT("/Engine/Functions/Engine_MaterialFunctions02/SmoothStep.SmoothStep"));
+				BaseColorFactor->ConnectExpression(*MultiplyExpression->GetInput(1), (int)FPBRMapFactory::EChannel::Alpha);
+				BaseColorMap->ConnectExpression(*MultiplyExpression->GetInput(0), (int)FPBRMapFactory::EChannel::Alpha);
 
-			FMaterialExpressionScalar* ValueExpression = MaterialElement.AddMaterialExpression<FMaterialExpressionScalar>();
-			ValueExpression->SetName(TEXT("Alpha Cuttof"));
-			ValueExpression->SetGroupName(GroupName);
-			ValueExpression->GetScalar() = GLTFMaterial.AlphaCutoff;
+				FMaterialExpressionFunctionCall* CuttofExpression = MaterialElement.AddMaterialExpression<FMaterialExpressionFunctionCall>();
+				CuttofExpression->SetFunctionPathName(TEXT("/Engine/Functions/Engine_MaterialFunctions02/SmoothStep.SmoothStep"));
 
-			MultiplyExpression->ConnectExpression(*CuttofExpression->GetInput(0), 0);
-			ValueExpression->ConnectExpression(*CuttofExpression->GetInput(1), 0);
-			ValueExpression->ConnectExpression(*CuttofExpression->GetInput(2), 0);
+				FMaterialExpressionScalar* ValueExpression = MaterialElement.AddMaterialExpression<FMaterialExpressionScalar>();
+				ValueExpression->SetName(TEXT("Alpha Cuttof"));
+				ValueExpression->SetGroupName(GroupName);
+				ValueExpression->GetScalar() = GLTFMaterial.AlphaCutoff;
 
-			CuttofExpression->ConnectExpression(MaterialElement.GetOpacity(), 0);
-			break;
-		}
-		case FMaterial::EAlphaMode::Blend:
-		{
-			FMaterialExpressionScalar* ValueExpression = MaterialElement.AddMaterialExpression<FMaterialExpressionScalar>();
-			ValueExpression->SetName(TEXT("IOR"));
-			ValueExpression->SetGroupName(GroupName);
-			ValueExpression->GetScalar() = 1.f;
-			ValueExpression->ConnectExpression(MaterialElement.GetRefraction(), 0);
+				MultiplyExpression->ConnectExpression(*CuttofExpression->GetInput(0), 0);
+				ValueExpression->ConnectExpression(*CuttofExpression->GetInput(1), 0);
+				ValueExpression->ConnectExpression(*CuttofExpression->GetInput(2), 0);
 
-			FMaterialExpressionColor* BaseColorFactor = FindExpression<FMaterialExpressionColor>(TEXT("BaseColor"), MaterialElement);
-			if (BaseColorMap)
-			{
-				FMaterialExpressionGeneric* MultiplyExpression = MaterialElement.AddMaterialExpression<FMaterialExpressionGeneric>();
-				MultiplyExpression->SetExpressionName(TEXT("Multiply"));
-				BaseColorFactor->ConnectExpression(*MultiplyExpression->GetInput(1), (int)FPBRMapFactory::EChannel::Alpha);
-				BaseColorMap->ConnectExpression(*MultiplyExpression->GetInput(0), (int)FPBRMapFactory::EChannel::Alpha);
-				MultiplyExpression->ConnectExpression(MaterialElement.GetOpacity(), 0);
+				CuttofExpression->ConnectExpression(MaterialElement.GetOpacity(), 0);
+				break;
 			}
-			else
+			case FMaterial::EAlphaMode::Blend:
 			{
-				BaseColorFactor->ConnectExpression(MaterialElement.GetOpacity(), (int)FPBRMapFactory::EChannel::Alpha);
+				FMaterialExpressionScalar* ValueExpression = MaterialElement.AddMaterialExpression<FMaterialExpressionScalar>();
+				ValueExpression->SetName(TEXT("IOR"));
+				ValueExpression->SetGroupName(GroupName);
+				ValueExpression->GetScalar() = 1.f;
+				ValueExpression->ConnectExpression(MaterialElement.GetRefraction(), 0);
+
+				FMaterialExpressionColor* BaseColorFactor = FindExpression<FMaterialExpressionColor>(TEXT("BaseColor"), MaterialElement);
+				if (BaseColorMap)
+				{
+					FMaterialExpressionGeneric* MultiplyExpression = MaterialElement.AddMaterialExpression<FMaterialExpressionGeneric>();
+					MultiplyExpression->SetExpressionName(TEXT("Multiply"));
+
+					BaseColorFactor->ConnectExpression(*MultiplyExpression->GetInput(1), (int)FPBRMapFactory::EChannel::Alpha);
+					BaseColorMap->ConnectExpression(*MultiplyExpression->GetInput(0), (int)FPBRMapFactory::EChannel::Alpha);
+					MultiplyExpression->ConnectExpression(MaterialElement.GetOpacity(), 0);
+				}
+				else
+				{
+					BaseColorFactor->ConnectExpression(MaterialElement.GetOpacity(), (int)FPBRMapFactory::EChannel::Alpha);
+				}
+				break;
 			}
-			break;
-		}
-		default:
-			check(false);
-			break;
+			default:
+				check(false);
+				break;
 		}
 	}
 
@@ -342,6 +348,8 @@ namespace GLTF
 			return;
 		}
 
+		MaterialElement.SetShadingModel(EGLTFMaterialShadingModel::ClearCoat);
+
 		FMaterialExpressionScalar* ClearCoatFactor = MaterialElement.AddMaterialExpression<FMaterialExpressionScalar>();
 		ClearCoatFactor->GetScalar() = GLTFMaterial.ClearCoat.ClearCoatFactor;
 		ClearCoatFactor->SetName(TEXT("ClearCoatFactor"));
@@ -367,6 +375,7 @@ namespace GLTF
 		{
 			FMaterialExpressionGeneric* MultiplyExpr = MaterialElement.AddMaterialExpression<FMaterialExpressionGeneric>();
 			MultiplyExpr->SetExpressionName(TEXT("Multiply"));
+
 			ClearCoatFactor->ConnectExpression(*MultiplyExpr->GetInput(0), 0);
 			ClearCoatTexture->ConnectExpression(*MultiplyExpr->GetInput(1), (int)FPBRMapFactory::EChannel::Red);
 			ClearCoatExpr = MultiplyExpr;
@@ -377,6 +386,7 @@ namespace GLTF
 		{
 			FMaterialExpressionGeneric* MultiplyExpr = MaterialElement.AddMaterialExpression<FMaterialExpressionGeneric>();
 			MultiplyExpr->SetExpressionName(TEXT("Multiply"));
+
 			ClearCoatRoughnessFactor->ConnectExpression(*MultiplyExpr->GetInput(0), 0);
 			ClearCoatRoughnessTexture->ConnectExpression(*MultiplyExpr->GetInput(1), (int)FPBRMapFactory::EChannel::Green);
 			ClearCoatRougnessExpr = MultiplyExpr;
@@ -393,8 +403,22 @@ namespace GLTF
 
 		if (ClearCoatNormalTexture)
 		{
+			FMaterialExpressionScalar* UVScale = MaterialElement.AddMaterialExpression<FMaterialExpressionScalar>();
+			UVScale->GetScalar() = GLTFMaterial.ClearCoat.NormalMapUVScale;
+			UVScale->SetName(TEXT("NormalMapUVScale"));
+
+			FMaterialExpressionGeneric* Mult = MaterialElement.AddMaterialExpression<FMaterialExpressionGeneric>();
+			Mult->SetExpressionName(TEXT("Multiply"));
+
+			FMaterialExpressionTextureCoordinate* TexCoord = MaterialElement.AddMaterialExpression<FMaterialExpressionTextureCoordinate>();
+			TexCoord->SetCoordinateIndex(GLTFMaterial.ClearCoat.NormalMap.TexCoord);
+			
 			FMaterialExpressionGeneric* ClearCoatNormalCustomOutput = MaterialElement.AddMaterialExpression<FMaterialExpressionGeneric>();
 			ClearCoatNormalCustomOutput->SetExpressionName(TEXT("ClearCoatNormalCustomOutput"));
+			
+			UVScale->ConnectExpression(*Mult->GetInput(0), 0);
+			TexCoord->ConnectExpression(*Mult->GetInput(1), 0);
+			Mult->ConnectExpression(*ClearCoatNormalTexture->GetInput(0), 0);
 			ClearCoatNormalTexture->ConnectExpression(*ClearCoatNormalCustomOutput->GetInput(0), 0);
 		}
 	}
@@ -409,52 +433,159 @@ namespace GLTF
 		MaterialElement.SetBlendMode(ConvertAlphaMode(GLTF::FMaterial::EAlphaMode::Blend));
 		MaterialElement.SetShadingModel(GLTF::EGLTFMaterialShadingModel::ThinTranslucent);
 		MaterialElement.SetTranslucencyLightingMode(ETranslucencyLightingMode::TLM_SurfacePerPixelLighting);
+		MaterialElement.SetTwoSided(GLTFMaterial.bIsDoubleSided);
 
 		FMaterialExpressionGeneric* ThinTranslucentOutput = MaterialElement.AddMaterialExpression<FMaterialExpressionGeneric>();
 		ThinTranslucentOutput->SetExpressionName(TEXT("ThinTranslucentMaterialOutput"));
 
-		// Connect whatever base color we have to the ThinTranslucentMaterialOutput node
-		FMaterialExpression* BaseColorExpr = MaterialElement.GetBaseColor().GetExpression();
-		BaseColorExpr->ConnectExpression(*ThinTranslucentOutput->GetInput(0), MaterialElement.GetBaseColor().GetOutputIndex());
-
 		FMaterialExpressionScalar* TransmissionFactorExpr = MaterialElement.AddMaterialExpression<FMaterialExpressionScalar>();
 		TransmissionFactorExpr->GetScalar() = GLTFMaterial.Transmission.TransmissionFactor;
 		TransmissionFactorExpr->SetName(TEXT("TransmissionFactor"));
 
-		FMaterialExpression* TransmissionExpr = TransmissionFactorExpr;
-
 		FMaterialExpressionTexture* TransmissionTexture = MapFactory.CreateTextureMap(
-			GetTexture(GLTFMaterial.Transmission.TransmissionMap, Textures), 
-			GLTFMaterial.Transmission.TransmissionMap.TexCoord, 
-			TEXT("Transmission"), 
+			GetTexture(GLTFMaterial.Transmission.TransmissionMap, Textures),
+			GLTFMaterial.Transmission.TransmissionMap.TexCoord,
+			TEXT("Transmission"),
 			ETextureMode::Color);
 
-		if (TransmissionTexture)
-		{
-			FMaterialExpressionGeneric* MultiplyExpr = MaterialElement.AddMaterialExpression<FMaterialExpressionGeneric>();
-			MultiplyExpr->SetExpressionName(TEXT("Multiply"));
-			TransmissionFactorExpr->ConnectExpression(*MultiplyExpr->GetInput(0), 0);
-			TransmissionTexture->ConnectExpression(*MultiplyExpr->GetInput(1), (int)FPBRMapFactory::EChannel::Red);
-			TransmissionExpr = MultiplyExpr;
-		}
+		FMaterialExpressionTexture* BaseColorMap = FindExpression<FMaterialExpressionTexture>(TEXT("BaseColor Map"), MaterialElement);
 
-		// Connect to opacity
-		if (FMaterialExpression* OpacityInputExpr = MaterialElement.GetOpacity().GetExpression())
+		if (GLTFMaterial.AlphaMode == FMaterial::EAlphaMode::Mask)
 		{
-			const int32 OutputIndex = MaterialElement.GetOpacity().GetOutputIndex();
-			FMaterialExpressionGeneric* MultiplyExpr = MaterialElement.AddMaterialExpression<FMaterialExpressionGeneric>();
-			MultiplyExpr->SetExpressionName(TEXT("Multiply"));
-			OpacityInputExpr->ConnectExpression(*MultiplyExpr->GetInput(0), OutputIndex);
-			TransmissionExpr->ConnectExpression(*MultiplyExpr->GetInput(1), 0);
-			TransmissionExpr = MultiplyExpr;
-		}
+			// Build opacity mask graph
 
-		if (TransmissionExpr == TransmissionFactorExpr)
-		{
-			TransmissionFactorExpr->GetScalar() = 1.0f - GLTFMaterial.Transmission.TransmissionFactor;
+			FMaterialExpressionColor* BaseColorFactor = FindExpression<FMaterialExpressionColor>(TEXT("BaseColor"), MaterialElement);
+
+			FMaterialExpression* BaseColor1 = nullptr;
+			FMaterialExpression* BaseColor2 = nullptr;
+
+			if (BaseColorMap)
+			{
+				BaseColor1 = BaseColorMap;
+			}
+			else
+			{
+				FMaterialExpressionColor* DefaultColor = MaterialElement.AddMaterialExpression<FMaterialExpressionColor>();
+				DefaultColor->GetColor() = FLinearColor::White;
+				BaseColor1 = DefaultColor;
+			}
+
+			if (BaseColorFactor)
+			{
+				BaseColor2 = BaseColorFactor;
+			}
+			else
+			{
+				FMaterialExpressionColor* DefaultColor = MaterialElement.AddMaterialExpression<FMaterialExpressionColor>();
+				DefaultColor->GetColor() = FLinearColor::White;
+				BaseColor2 = DefaultColor;
+			}
+
+			FMaterialExpressionGeneric* Mult1 = MaterialElement.AddMaterialExpression<FMaterialExpressionGeneric>();
+			Mult1->SetExpressionName(TEXT("Multiply"));
+
+			BaseColor1->ConnectExpression(*Mult1->GetInput(0), 0);
+			BaseColor2->ConnectExpression(*Mult1->GetInput(1), 0);
+
+			FMaterialExpressionGeneric* Step1 = MaterialElement.AddMaterialExpression<FMaterialExpressionGeneric>();
+			Step1->SetExpressionName(TEXT("Step"));
+			Step1->SetFloatProperty(TEXT("ConstX"), 0.333f);
+
+			FMaterialExpressionGeneric* OneMinus1 = MaterialElement.AddMaterialExpression<FMaterialExpressionGeneric>();
+			OneMinus1->SetExpressionName(TEXT("OneMinus"));
+			
+			BaseColor1->ConnectExpression(*Step1->GetInput(0), (int)FPBRMapFactory::EChannel::Alpha);
+			Step1->ConnectExpression(*OneMinus1->GetInput(0), 0);
+			Step1->ConnectExpression(MaterialElement.GetSpecular(), 0);
+
+			FMaterialExpressionGeneric* Lerp1 = MaterialElement.AddMaterialExpression<FMaterialExpressionGeneric>();
+			Lerp1->SetExpressionName(TEXT("LinearInterpolate"));
+			OneMinus1->ConnectExpression(*Lerp1->GetInput(2), 0);
+
+			FMaterialExpressionGeneric* Mult2 = MaterialElement.AddMaterialExpression<FMaterialExpressionGeneric>();
+			Mult2->SetExpressionName(TEXT("Multiply"));
+
+			Mult2->ConnectExpression(MaterialElement.GetOpacity(), 0);
+
+			if (TransmissionTexture)
+			{
+				FMaterialExpressionGeneric* Mult3 = MaterialElement.AddMaterialExpression<FMaterialExpressionGeneric>();
+				Mult3->SetExpressionName(TEXT("Multiply"));
+
+				FMaterialExpressionGeneric* Mult4 = MaterialElement.AddMaterialExpression<FMaterialExpressionGeneric>();
+				Mult4->SetExpressionName(TEXT("Multiply"));
+
+				FMaterialExpressionGeneric* OneMinus2 = MaterialElement.AddMaterialExpression<FMaterialExpressionGeneric>();
+				OneMinus2->SetExpressionName(TEXT("OneMinus"));
+
+				Mult1->ConnectExpression(*Mult3->GetInput(0), 0);
+
+				Mult3->ConnectExpression(*Mult4->GetInput(0), 0);
+				Mult3->ConnectExpression(*OneMinus2->GetInput(0), 0);
+				Mult1->ConnectExpression(*Mult4->GetInput(1), 0);
+				Mult4->ConnectExpression(*Lerp1->GetInput(0), 0);
+
+				Step1->ConnectExpression(*Mult2->GetInput(0), 0);
+				OneMinus2->ConnectExpression(*Mult2->GetInput(1), 0);
+				
+				TransmissionTexture->ConnectExpression(*Mult3->GetInput(0), 1);
+				TransmissionFactorExpr->ConnectExpression(*Mult3->GetInput(1), 0);
+				Lerp1->ConnectExpression(*ThinTranslucentOutput->GetInput(0), 0);
+			}
+			else
+			{
+				FMaterialExpressionScalar* Const1 = MaterialElement.AddMaterialExpression<FMaterialExpressionScalar>();
+				Const1->GetScalar() = 0;
+
+				Step1->ConnectExpression(*Mult2->GetInput(1), 0);
+				Const1->ConnectExpression(*Mult2->GetInput(0), 0);
+
+				Mult1->ConnectExpression(*Lerp1->GetInput(0), 0);
+
+				FMaterialExpressionGeneric* Mult3 = MaterialElement.AddMaterialExpression<FMaterialExpressionGeneric>();
+				Mult3->SetExpressionName(TEXT("Multiply"));
+
+				Lerp1->ConnectExpression(*Mult3->GetInput(0), 0);
+				TransmissionFactorExpr->ConnectExpression(*Mult3->GetInput(1), 0);
+
+				Mult3->ConnectExpression(*ThinTranslucentOutput->GetInput(0), 0);
+			}
+
+			BaseColor2->ConnectExpression(MaterialElement.GetBaseColor(), 0);
 		}
+		else
+		{
+			FMaterialExpressionGeneric* BaseColorMult = MaterialElement.AddMaterialExpression<FMaterialExpressionGeneric>();
+			BaseColorMult->SetExpressionName(TEXT("Multiply"));
+
+			BaseColorMult->ConnectExpression(*ThinTranslucentOutput->GetInput(0), 0);
+
+			FMaterialExpressionGeneric* OpacitySubtract = MaterialElement.AddMaterialExpression<FMaterialExpressionGeneric>();
+			OpacitySubtract->SetExpressionName(TEXT("Subtract"));
+
+			FMaterialExpression* BaseColorExpr = MaterialElement.GetBaseColor().GetExpression();
+			if (BaseColorExpr)
+			{
+				BaseColorExpr->ConnectExpression(*BaseColorMult->GetInput(0), MaterialElement.GetBaseColor().GetOutputIndex());
+			}
 
-		TransmissionExpr->ConnectExpression(MaterialElement.GetOpacity(), 0);
+			if (TransmissionTexture)
+			{
+				FMaterialExpressionGeneric* TransmissionMult = MaterialElement.AddMaterialExpression<FMaterialExpressionGeneric>();
+				TransmissionMult->SetExpressionName(TEXT("Multiply"));
+				TransmissionFactorExpr->ConnectExpression(*TransmissionMult->GetInput(0), 0);
+				TransmissionTexture->ConnectExpression(*TransmissionMult->GetInput(1), (int)FPBRMapFactory::EChannel::Red);
+				TransmissionMult->ConnectExpression(*BaseColorMult->GetInput(1), 0);
+				TransmissionMult->ConnectExpression(*OpacitySubtract->GetInput(1), 0);
+			}
+			else
+			{
+				TransmissionFactorExpr->ConnectExpression(*BaseColorMult->GetInput(1), 0);
+				TransmissionFactorExpr->ConnectExpression(*OpacitySubtract->GetInput(1), 0);
+			}
+
+			OpacitySubtract->ConnectExpression(MaterialElement.GetOpacity(), 0);
+		}
 	}
 	
 	void FMaterialFactoryImpl::HandleSheen(const TArray<GLTF::FTexture>& Textures, const GLTF::FMaterial& GLTFMaterial, FPBRMapFactory& MapFactory, FMaterialElement& MaterialElement)
@@ -494,6 +625,72 @@ namespace GLTF
 		}
 	}
 
+	void FMaterialFactoryImpl::HandleIOR(const TArray<GLTF::FTexture>& Textures, const GLTF::FMaterial& GLTFMaterial, FPBRMapFactory& MapFactory, FMaterialElement& MaterialElement)
+	{
+		if (!GLTFMaterial.bHasIOR || GLTFMaterial.ShadingModel == FMaterial::EShadingModel::SpecularGlossiness || GLTFMaterial.bIsUnlitShadingModel)
+		{
+			// As per glTF 2.0 spec, IOR extension is not compatible with specular-glossiness workflow or unlit shading model
+			return;
+		}
+
+		if (!GLTFMaterial.bHasTransmission)
+		{
+			MaterialElement.SetBlendMode(ConvertAlphaMode(GLTF::FMaterial::EAlphaMode::Blend));
+			MaterialElement.SetTranslucencyLightingMode(ETranslucencyLightingMode::TLM_Surface);
+			MaterialElement.SetTwoSided(GLTFMaterial.bIsDoubleSided);
+		}
+
+		FMaterialExpressionGeneric* LerpExpr = MaterialElement.AddMaterialExpression<FMaterialExpressionGeneric>();
+		LerpExpr->SetExpressionName(TEXT("LinearInterpolate"));
+
+		FMaterialExpressionScalar* ConstExpr = MaterialElement.AddMaterialExpression<FMaterialExpressionScalar>();
+		ConstExpr->GetScalar() = 1.f;
+
+		FMaterialExpressionScalar* IORExpr = MaterialElement.AddMaterialExpression<FMaterialExpressionScalar>();
+		IORExpr->SetName(TEXT("IOR"));
+		IORExpr->GetScalar() = GLTFMaterial.IOR;
+
+		FMaterialExpressionGeneric* FresnelExpr = MaterialElement.AddMaterialExpression<FMaterialExpressionGeneric>();
+		FresnelExpr->SetExpressionName(TEXT("Fresnel"));
+
+		ConstExpr->ConnectExpression(*LerpExpr->GetInput(0), 0);
+		IORExpr->ConnectExpression(*LerpExpr->GetInput(1), 0);
+		FresnelExpr->ConnectExpression(*LerpExpr->GetInput(2), 0);
+
+		LerpExpr->ConnectExpression(MaterialElement.GetRefraction(), 0);
+	}
+
+	void FMaterialFactoryImpl::HandleSpecular(const TArray<GLTF::FTexture>& Textures, const GLTF::FMaterial& GLTFMaterial, FPBRMapFactory& MapFactory, FMaterialElement& MaterialElement)
+	{
+		if (!GLTFMaterial.bHasSpecular || GLTFMaterial.ShadingModel == FMaterial::EShadingModel::SpecularGlossiness || GLTFMaterial.bIsUnlitShadingModel)
+		{
+			// As per glTF 2.0 spec, specular extension is not compatible with specular-glossiness workflow or unlit shading model
+			return;
+		}
+
+		FMaterialExpressionScalar* SpecularFactorExpr = MaterialElement.AddMaterialExpression<FMaterialExpressionScalar>();
+		SpecularFactorExpr->SetName(TEXT("SpecularFactor"));
+		SpecularFactorExpr->GetScalar() = GLTFMaterial.Specular.SpecularFactor;
+
+		FMaterialExpressionTexture* SpecularTextureExpr = MapFactory.CreateTextureMap(
+			GetTexture(GLTFMaterial.Specular.SpecularMap, Textures), GLTFMaterial.Specular.SpecularMap.TexCoord, TEXT("SpecularMap"), ETextureMode::Color);
+
+		FMaterialExpression* SpecularExpr = SpecularFactorExpr;
+
+		if (SpecularTextureExpr)
+		{
+			FMaterialExpressionGeneric* MultiplyExpr = MaterialElement.AddMaterialExpression<FMaterialExpressionGeneric>();
+			MultiplyExpr->SetExpressionName(TEXT("Multiply"));
+			
+			SpecularFactorExpr->ConnectExpression(*MultiplyExpr->GetInput(0), 0);
+			SpecularTextureExpr->ConnectExpression(*MultiplyExpr->GetInput(1), (int)FPBRMapFactory::EChannel::Alpha);
+
+			SpecularExpr = MultiplyExpr;
+		}
+
+		SpecularExpr->ConnectExpression(MaterialElement.GetSpecular(), 0);
+	}
+
 	FMaterialFactory::FMaterialFactory(IMaterialElementFactory* MaterialElementFactory, ITextureFactory* TextureFactory)
 	    : Impl(new FMaterialFactoryImpl(MaterialElementFactory, TextureFactory))
 	{
diff --git a/Engine/Plugins/Editor/GLTFImporter/Source/GLTFCore/Public/GLTFAsset.h b/Engine/Plugins/Editor/GLTFImporter/Source/GLTFCore/Public/GLTFAsset.h
--- a/Engine/Plugins/Editor/GLTFImporter/Source/GLTFCore/Public/GLTFAsset.h
+++ b/Engine/Plugins/Editor/GLTFImporter/Source/GLTFCore/Public/GLTFAsset.h
@@ -21,6 +21,8 @@ namespace GLTF
 		KHR_MaterialsTransmission,
 		KHR_MaterialsSheen,
 		KHR_MaterialsVariants,
+		KHR_MaterialsIOR,
+		KHR_MaterialsSpecular,
 		KHR_TextureTransform,
 		KHR_DracoMeshCompression,
 		KHR_LightsPunctual,
diff --git a/Engine/Plugins/Editor/GLTFImporter/Source/GLTFCore/Public/GLTFMaterial.h b/Engine/Plugins/Editor/GLTFImporter/Source/GLTFCore/Public/GLTFMaterial.h
--- a/Engine/Plugins/Editor/GLTFImporter/Source/GLTFCore/Public/GLTFMaterial.h
+++ b/Engine/Plugins/Editor/GLTFImporter/Source/GLTFCore/Public/GLTFMaterial.h
@@ -61,68 +61,50 @@ namespace GLTF
 		struct FMetallicRoughness
 		{
 			FTextureMap Map;
-			float       MetallicFactor;
-			float       RoughnessFactor;
-
-			FMetallicRoughness()
-			    : MetallicFactor(1.f)
-			    , RoughnessFactor(1.f)
-			{
-			}
+			float       MetallicFactor = 1.0f;
+			float       RoughnessFactor = 1.0f;
 		};
-		struct FSpecularGlossiness
+		// "pbrSpecularGlossiness" extension
+		// Implements a specular glossiness shading model and is mutually exclusive with "specular" extension, which uses metallic roughness model
+		struct FSpecularGlossiness 
 		{
 			FTextureMap Map;
-			FVector     SpecularFactor;
-			float       GlossinessFactor;
-
-			FSpecularGlossiness()
-			    : SpecularFactor(1.f)
-			    , GlossinessFactor(1.f)
-			{
-			}
+			FVector     SpecularFactor = FVector(1.0f);
+			float       GlossinessFactor = 1.0f;
 		};
 		struct FClearCoat
 		{
-			float ClearCoatFactor;
+			float ClearCoatFactor = 1.0f;
 			FTextureMap ClearCoatMap;
 
-			float Roughness;
+			float Roughness = 0.0f;
 			FTextureMap RoughnessMap;
 
+			float NormalMapUVScale = 1.0f;
 			FTextureMap NormalMap;
-
-			FClearCoat()
-				: ClearCoatFactor(1.0f)
-				, Roughness(0.0f)
-			{
-			}
 		};
 		struct FTransmission
 		{
-			float TransmissionFactor;
+			float TransmissionFactor = 0.0f;
 			FTextureMap TransmissionMap;
-
-			FTransmission()
-				: TransmissionFactor(0.0f)
-			{
-			}
 		};
 		struct FSheen
 		{
 			FVector SheenColorFactor = FVector::Zero();
 		};
+		struct FSpecular // "specular" extension
+		{
+			float SpecularFactor = 1.0f;
+			FTextureMap SpecularMap;
+			FVector SpecularColorFactor = FVector(1.0f, 1.0f, 1.0f);
+			FTextureMap SpecularColorMap;
+		};
 
 		struct FPacking
 		{
-			int         Flags;  // see EPackingFlags
+			int         Flags = (int)EPackingFlags::None;  // see EPackingFlags
 			FTextureMap Map;
 			FTextureMap NormalMap;
-
-			FPacking()
-			    : Flags((int)EPackingFlags::None)
-			{
-			}
 		};
 
 		FString Name;
@@ -133,9 +115,11 @@ namespace GLTF
 		EShadingModel       ShadingModel;
 		FMetallicRoughness  MetallicRoughness;
 		FSpecularGlossiness SpecularGlossiness;
-		FClearCoat			ClearCoat;
-		FTransmission		Transmission;
-		FSheen				Sheen;
+		FClearCoat          ClearCoat;
+		FTransmission       Transmission;
+		FSheen              Sheen;
+		FSpecular           Specular;
+		float               IOR;
 
 		// base properties
 		FTextureMap Normal;
@@ -156,11 +140,14 @@ namespace GLTF
 		bool     bHasClearCoat;
 		bool     bHasSheen;
 		bool     bHasTransmission;
+		bool     bHasIOR;
+		bool     bHasSpecular;
 
 		FMaterial(const FString& Name)
 		    : Name(Name)
 		    , BaseColorFactor {1.0f, 1.0f, 1.0f, 1.0f}
 		    , ShadingModel(EShadingModel::MetallicRoughness)
+			, IOR(1.0f)
 		    , NormalScale(1.f)
 		    , OcclusionStrength(1.f)
 		    , EmissiveFactor(FVector::ZeroVector)
@@ -171,6 +158,8 @@ namespace GLTF
 			, bHasClearCoat(false)
 			, bHasSheen(false)
 			, bHasTransmission(false)
+			, bHasIOR(false)
+			, bHasSpecular(false)
 		{
 		}
 
diff --git a/Engine/Plugins/Editor/GLTFImporter/Source/GLTFCore/Public/GLTFMaterialExpressions.h b/Engine/Plugins/Editor/GLTFImporter/Source/GLTFCore/Public/GLTFMaterialExpressions.h
--- a/Engine/Plugins/Editor/GLTFImporter/Source/GLTFCore/Public/GLTFMaterialExpressions.h
+++ b/Engine/Plugins/Editor/GLTFImporter/Source/GLTFCore/Public/GLTFMaterialExpressions.h
@@ -369,7 +369,8 @@ namespace GLTF
 
 		TArray<FMaterialExpression*> Expressions;
 
-		FMaterialExpression* ThinTranslucentMaterialOutput;
+		GLTF::FMaterialExpression* ThinTranslucentMaterialOutput;
+		GLTF::FMaterialExpression* ClearCoatBottomNormalOutput;
 
 		bool bIsFinal;
 
diff --git a/Engine/Plugins/Editor/GLTFImporter/Source/GLTFImporter/Private/GLTFMaterialElement.cpp b/Engine/Plugins/Editor/GLTFImporter/Source/GLTFImporter/Private/GLTFMaterialElement.cpp
--- a/Engine/Plugins/Editor/GLTFImporter/Source/GLTFImporter/Private/GLTFMaterialElement.cpp
+++ b/Engine/Plugins/Editor/GLTFImporter/Source/GLTFImporter/Private/GLTFMaterialElement.cpp
@@ -15,6 +15,7 @@
 #include "Materials/MaterialExpressionTextureSampleParameter2D.h"
 #include "Materials/MaterialExpressionVectorParameter.h"
 #include "Materials/MaterialExpressionThinTranslucentMaterialOutput.h"
+#include "Materials/MaterialExpressionClearCoatNormalCustomOutput.h"
 
 namespace GLTFImporterImpl
 {
@@ -275,24 +276,41 @@ void FGLTFMaterialElement::Finalize()
 	ConnectInput(Refraction, MaterialExpressions, Material->Refraction);
 	ConnectInput(Normal, MaterialExpressions, Material->Normal);
 	ConnectInput(AmbientOcclusion, MaterialExpressions, Material->AmbientOcclusion);
+	ConnectInput(ClearCoat, MaterialExpressions, Material->ClearCoat);
+	ConnectInput(ClearCoatRoughness, MaterialExpressions, Material->ClearCoatRoughness);
 
 	// Handle transmission materials (they add a special output node to the graph)
 	if (ThinTranslucentMaterialOutput)
 	{
-		GLTF::FMaterialExpressionInput* Input = ThinTranslucentMaterialOutput->GetInput(0);
-		check(Input);
+		const int32 ThinTranslucentExpressionIndex = Expressions.Find(ThinTranslucentMaterialOutput);
+		
+		if (ThinTranslucentExpressionIndex != INDEX_NONE && MaterialExpressions[ThinTranslucentExpressionIndex].IsValid())
+		{
+			UMaterialExpression& ThinTranslucentMaterialExpression = *MaterialExpressions[ThinTranslucentExpressionIndex];
+
+			GLTF::FMaterialExpressionInput* ThinTranslucentInput = ThinTranslucentMaterialOutput->GetInput(0);
+
+			if (ThinTranslucentInput)
+			{
+				ConnectInput(*ThinTranslucentInput, MaterialExpressions, *ThinTranslucentMaterialExpression.GetInput(0));
+			}
+		}
+	}
 
-		if (GLTF::FMaterialExpression* BaseColorExpr = Input->GetExpression())
+	if (ClearCoatBottomNormalOutput)
+	{
+		const int32 ClearCoatBottomNormalOutputIndex = Expressions.Find(ClearCoatBottomNormalOutput);
+
+		if (ClearCoatBottomNormalOutputIndex != INDEX_NONE && MaterialExpressions[ClearCoatBottomNormalOutputIndex].IsValid())
 		{
-			const int32 ThinTranslucentExpressionIndex = Expressions.Find(ThinTranslucentMaterialOutput);
-			check(ThinTranslucentExpressionIndex != INDEX_NONE);
-			TStrongObjectPtr<UMaterialExpression> ThinTranslucentMaterialExpression = MaterialExpressions[ThinTranslucentExpressionIndex];
+			UMaterialExpression& ClearCoatBottomNormalMaterialExpression = *MaterialExpressions[ClearCoatBottomNormalOutputIndex];
 
-			const int32 BaseColorExpressionIndex = Expressions.Find(BaseColorExpr);
-			check(BaseColorExpressionIndex != INDEX_NONE);
-			TStrongObjectPtr<UMaterialExpression> BaseColorMaterialExpression = MaterialExpressions[BaseColorExpressionIndex];
+			GLTF::FMaterialExpressionInput* ClearCoatBottomNormalInput = ClearCoatBottomNormalOutput->GetInput(0);
 
-			BaseColorMaterialExpression->ConnectExpression(ThinTranslucentMaterialExpression->GetInput(0), Input->GetOutputIndex());
+			if (ClearCoatBottomNormalInput)
+			{
+				ConnectInput(*ClearCoatBottomNormalInput, MaterialExpressions, *ClearCoatBottomNormalMaterialExpression.GetInput(0));
+			}
 		}
 	}
 
@@ -348,6 +366,11 @@ void FGLTFMaterialElement::CreateExpressions(TArray<TStrongObjectPtr<UMaterialEx
 		{
 			ThinTranslucentMaterialOutput = Expression;
 		}
+
+		if (MaterialExpression->GetClass() == UMaterialExpressionClearCoatNormalCustomOutput::StaticClass())
+		{
+			ClearCoatBottomNormalOutput = Expression;
+		}
 	}
 }
 
diff --git a/Engine/Plugins/Enterprise/DatasmithGLTFImporter/Source/DatasmithGLTFTranslator/Private/DatasmithGLTFMaterialElement.cpp b/Engine/Plugins/Enterprise/DatasmithGLTFImporter/Source/DatasmithGLTFTranslator/Private/DatasmithGLTFMaterialElement.cpp
--- a/Engine/Plugins/Enterprise/DatasmithGLTFImporter/Source/DatasmithGLTFTranslator/Private/DatasmithGLTFMaterialElement.cpp
+++ b/Engine/Plugins/Enterprise/DatasmithGLTFImporter/Source/DatasmithGLTFTranslator/Private/DatasmithGLTFMaterialElement.cpp
@@ -86,15 +86,45 @@ void FDatasmithGLTFMaterialElement::Finalize()
 	ConnectInput(Refraction, MaterialExpressions, MaterialElement->GetRefraction());
 	ConnectInput(Normal, MaterialExpressions, MaterialElement->GetNormal());
 	ConnectInput(AmbientOcclusion, MaterialExpressions, MaterialElement->GetAmbientOcclusion());
+	ConnectInput(ClearCoat, MaterialExpressions, MaterialElement->GetClearCoat());
+	ConnectInput(ClearCoatRoughness, MaterialExpressions, MaterialElement->GetClearCoatRoughness());
+
+	if (ThinTranslucentMaterialOutput)
+	{
+		const int32 ThinTranslucentExpressionIndex = Expressions.Find(ThinTranslucentMaterialOutput);
+		check(ThinTranslucentExpressionIndex != INDEX_NONE);
+		IDatasmithMaterialExpression* ThinTranslucentMaterialExpression = MaterialExpressions[ThinTranslucentExpressionIndex];
+
+		GLTF::FMaterialExpressionInput* ThinTranslucentInput = ThinTranslucentMaterialOutput->GetInput(0);
+
+		if (ThinTranslucentInput && ThinTranslucentMaterialExpression)
+		{
+			ConnectInput(*ThinTranslucentInput, MaterialExpressions, *ThinTranslucentMaterialExpression->GetInput(0));
+		}
+	}
+
+	if (ClearCoatBottomNormalOutput)
+	{
+		const int32 ClearCoatBottomNormalOutputIndex = Expressions.Find(ClearCoatBottomNormalOutput);
+		check(ClearCoatBottomNormalOutputIndex != INDEX_NONE);
+		IDatasmithMaterialExpression* ClearCoatBottomNormalMaterialExpression = MaterialExpressions[ClearCoatBottomNormalOutputIndex];
+
+		GLTF::FMaterialExpressionInput* ClearCoatBottomNormalInput = ClearCoatBottomNormalOutput->GetInput(0);
+
+		if (ClearCoatBottomNormalInput && ClearCoatBottomNormalMaterialExpression)
+		{
+			ConnectInput(*ClearCoatBottomNormalInput, MaterialExpressions, *ClearCoatBottomNormalMaterialExpression->GetInput(0));
+		}
+	}
 
 	bIsFinal = true;
 }
 
-void FDatasmithGLTFMaterialElement::CreateExpressions(TArray<IDatasmithMaterialExpression*>& MaterialExpressions) const
+void FDatasmithGLTFMaterialElement::CreateExpressions(TArray<IDatasmithMaterialExpression*>& MaterialExpressions)
 {
 	MaterialExpressions.Empty(Expressions.Num());
 
-	for (const GLTF::FMaterialExpression* Expression : Expressions)
+	for (GLTF::FMaterialExpression* Expression : Expressions)
 	{
 		using namespace DatasmithGLTFImporterImpl;
 
@@ -126,11 +156,22 @@ void FDatasmithGLTFMaterialElement::CreateExpressions(TArray<IDatasmithMaterialE
 			break;
 			case GLTF::EMaterialExpressionType::Generic:
 			{
-				const GLTF::FMaterialExpressionGeneric& GenericExpression = *static_cast<const GLTF::FMaterialExpressionGeneric*>(Expression);
-				IDatasmithMaterialExpressionGeneric*    NewExpression = MaterialElement->AddMaterialExpression<IDatasmithMaterialExpressionGeneric>();
-				NewExpression->SetExpressionName(GenericExpression.GetExpressionName());
+				GLTF::FMaterialExpressionGeneric* GenericExpression = static_cast<GLTF::FMaterialExpressionGeneric*>(Expression);
+				IDatasmithMaterialExpressionGeneric* NewExpression = MaterialElement->AddMaterialExpression<IDatasmithMaterialExpressionGeneric>();
+				const FString ExprName(GenericExpression->GetExpressionName());
+				NewExpression->SetExpressionName(*ExprName);
+
+				// Check for special expressions
+				if (ExprName == "ClearCoatNormalCustomOutput")
+				{
+					ClearCoatBottomNormalOutput = GenericExpression;
+				}
+				else if (ExprName == "ThinTranslucentMaterialOutput")
+				{
+					ThinTranslucentMaterialOutput = GenericExpression;
+				}
 
-				for (const TPair<FString, bool>& NameValue : GenericExpression.GetBoolProperties())
+				for (const TPair<FString, bool>& NameValue : GenericExpression->GetBoolProperties())
 				{
 					TSharedPtr<IDatasmithKeyValueProperty> PropertyPtr = FDatasmithSceneFactory::CreateKeyValueProperty(*NameValue.Key);
 					PropertyPtr->SetPropertyType(EDatasmithKeyValuePropertyType::Bool);
@@ -138,7 +179,7 @@ void FDatasmithGLTFMaterialElement::CreateExpressions(TArray<IDatasmithMaterialE
 					NewExpression->AddProperty(PropertyPtr);
 				}
 
-				for (const TPair<FString, float>& NameValue : GenericExpression.GetFloatProperties())
+				for (const TPair<FString, float>& NameValue : GenericExpression->GetFloatProperties())
 				{
 					TSharedPtr<IDatasmithKeyValueProperty> PropertyPtr = FDatasmithSceneFactory::CreateKeyValueProperty(*NameValue.Key);
 					PropertyPtr->SetPropertyType(EDatasmithKeyValuePropertyType::Float);
diff --git a/Engine/Plugins/Enterprise/DatasmithGLTFImporter/Source/DatasmithGLTFTranslator/Private/DatasmithGLTFMaterialElement.h b/Engine/Plugins/Enterprise/DatasmithGLTFImporter/Source/DatasmithGLTFTranslator/Private/DatasmithGLTFMaterialElement.h
--- a/Engine/Plugins/Enterprise/DatasmithGLTFImporter/Source/DatasmithGLTFTranslator/Private/DatasmithGLTFMaterialElement.h
+++ b/Engine/Plugins/Enterprise/DatasmithGLTFImporter/Source/DatasmithGLTFTranslator/Private/DatasmithGLTFMaterialElement.h
@@ -39,7 +39,7 @@ public:
 	TSharedRef<IDatasmithUEPbrMaterialElement> GetMaterial() const;
 
 private:
-	void CreateExpressions(TArray<IDatasmithMaterialExpression*>& MaterialExpressions) const;
+	void CreateExpressions(TArray<IDatasmithMaterialExpression*>& MaterialExpressions);
 
 	void ConnectInput(const GLTF::FMaterialExpressionInput& ExpressionInput, const TArray<IDatasmithMaterialExpression*>& MaterialExpressions,
 	                  IDatasmithExpressionInput& MaterialInput) const;
diff --git a/Engine/Plugins/Enterprise/DatasmithImporter/Source/DatasmithImporter/Private/DatasmithFileProducer.cpp b/Engine/Plugins/Enterprise/DatasmithImporter/Source/DatasmithImporter/Private/DatasmithFileProducer.cpp
--- a/Engine/Plugins/Enterprise/DatasmithImporter/Source/DatasmithImporter/Private/DatasmithFileProducer.cpp
+++ b/Engine/Plugins/Enterprise/DatasmithImporter/Source/DatasmithImporter/Private/DatasmithFileProducer.cpp
@@ -457,7 +457,7 @@ void UDatasmithFileProducer::SceneElementToWorld()
 		}
 	}
 
-	for ( TPair< int32, UMaterialInterface* >& AssetPair : ImportContextPtr->ImportedParentMaterials )
+	for ( TPair< uint32, UMaterialInterface* >& AssetPair : ImportContextPtr->ImportedParentMaterials )
 	{
 		if(AssetPair.Value)
 		{
diff --git a/Engine/Plugins/Enterprise/DatasmithImporter/Source/DatasmithImporter/Private/DatasmithMaterialExpressions.cpp b/Engine/Plugins/Enterprise/DatasmithImporter/Source/DatasmithImporter/Private/DatasmithMaterialExpressions.cpp
--- a/Engine/Plugins/Enterprise/DatasmithImporter/Source/DatasmithImporter/Private/DatasmithMaterialExpressions.cpp
+++ b/Engine/Plugins/Enterprise/DatasmithImporter/Source/DatasmithImporter/Private/DatasmithMaterialExpressions.cpp
@@ -3155,6 +3155,18 @@ void FDatasmithMaterialExpressions::CreateUEPbrMaterialGraph(const TSharedPtr< I
 
 	ConnectExpression( MaterialElement.ToSharedRef(), MaterialExpressions, MaterialElement->GetNormal().GetExpression(), GetMaterialOrFunctionSlot( UnrealMaterialOrFunction, EDatasmithTextureSlot::NORMAL ), MaterialElement->GetNormal().GetOutputIndex() );
 
+	if ( MaterialElement->GetShadingModel() == EDatasmithShadingModel::ClearCoat )
+	{
+		if ( MaterialElement->GetClearCoat().GetExpression() )
+		{
+			ConnectExpression(MaterialElement.ToSharedRef(), MaterialExpressions, MaterialElement->GetClearCoat().GetExpression(), GetMaterialOrFunctionSlot(UnrealMaterialOrFunction, EDatasmithTextureSlot::COATSPECULAR), MaterialElement->GetClearCoat().GetOutputIndex());
+		}
+		if (MaterialElement->GetClearCoatRoughness().GetExpression())
+		{
+			ConnectExpression(MaterialElement.ToSharedRef(), MaterialExpressions, MaterialElement->GetClearCoatRoughness().GetExpression(), GetMaterialOrFunctionSlot(UnrealMaterialOrFunction, EDatasmithTextureSlot::COATROUGHNESS), MaterialElement->GetClearCoat().GetOutputIndex());
+		}
+	}
+
 	// Connect expressions to any UMaterialExpressionCustomOutput since these aren't part of the predefined material outputs
 	for ( int32 ExpressionIndex = 0; ExpressionIndex < MaterialElement->GetExpressionsCount(); ++ExpressionIndex )
 	{
diff --git a/Engine/Plugins/Enterprise/DatasmithImporter/Source/DatasmithImporter/Public/DatasmithImportContext.h b/Engine/Plugins/Enterprise/DatasmithImporter/Source/DatasmithImporter/Public/DatasmithImportContext.h
--- a/Engine/Plugins/Enterprise/DatasmithImporter/Source/DatasmithImporter/Public/DatasmithImportContext.h
+++ b/Engine/Plugins/Enterprise/DatasmithImporter/Source/DatasmithImporter/Public/DatasmithImportContext.h
@@ -271,7 +271,7 @@ struct DATASMITHIMPORTER_API FDatasmithImportContext
 	TMap< FString, TSharedRef < IDatasmithBaseMaterialElement > > ImportedMaterialFunctionsByName;
 
 	/** List of potential parent materials with their hash. Used to create material instances from. */
-	TMap< int32, UMaterialInterface* > ImportedParentMaterials;
+	TMap< uint32, UMaterialInterface* > ImportedParentMaterials;
 
 	/** Map of imported level sequences for each level sequence element */
 	TMap< TSharedRef< IDatasmithLevelSequenceElement >, ULevelSequence* > ImportedLevelSequences;
diff --git a/Engine/Source/Runtime/Datasmith/DatasmithCore/Private/DatasmithMaterialElementsImpl.cpp b/Engine/Source/Runtime/Datasmith/DatasmithCore/Private/DatasmithMaterialElementsImpl.cpp
--- a/Engine/Source/Runtime/Datasmith/DatasmithCore/Private/DatasmithMaterialElementsImpl.cpp
+++ b/Engine/Source/Runtime/Datasmith/DatasmithCore/Private/DatasmithMaterialElementsImpl.cpp
@@ -201,6 +201,8 @@ FDatasmithUEPbrMaterialElementImpl::FDatasmithUEPbrMaterialElementImpl( const TC
 	, WorldDisplacement(  MakeShared< FDatasmithExpressionInputImpl >( TEXT("WorldDisplacement") ) )
 	, Refraction(         MakeShared< FDatasmithExpressionInputImpl >( TEXT("Refraction") ) )
 	, AmbientOcclusion(   MakeShared< FDatasmithExpressionInputImpl >( TEXT("AmbientOcclusion") ) )
+	, ClearCoat(		  MakeShared< FDatasmithExpressionInputImpl >( TEXT("ClearCoat") ) )
+	, ClearCoatRoughness( MakeShared< FDatasmithExpressionInputImpl >( TEXT("ClearCoatRoughness") ) )
 	, MaterialAttributes( MakeShared< FDatasmithExpressionInputImpl >( TEXT("MaterialAttributes") ) )
 	, BlendMode(0)
 	, bTwoSided( false )
@@ -220,6 +222,8 @@ FDatasmithUEPbrMaterialElementImpl::FDatasmithUEPbrMaterialElementImpl( const TC
 	RegisterReferenceProxy( WorldDisplacement, "WorldDisplacement" );
 	RegisterReferenceProxy( Refraction, "Refraction" );
 	RegisterReferenceProxy( AmbientOcclusion, "AmbientOcclusion" );
+	RegisterReferenceProxy( ClearCoat, "ClearCoat" );
+	RegisterReferenceProxy( ClearCoatRoughness, "ClearCoatRoughness" );
 	RegisterReferenceProxy( MaterialAttributes, "MaterialAttributes" );
 
 	RegisterReferenceProxy( Expressions, "Expressions" );
@@ -271,6 +275,8 @@ FMD5Hash FDatasmithUEPbrMaterialElementImpl::CalculateElementHash(bool bForce)
 	UpdateMD5(*WorldDisplacement.Edit());
 	UpdateMD5(*Refraction.Edit());
 	UpdateMD5(*AmbientOcclusion.Edit());
+	UpdateMD5(*ClearCoat.Edit());
+	UpdateMD5(*ClearCoatRoughness.Edit());
 	UpdateMD5(*MaterialAttributes.Edit());
 
 	ElementHash.Set(MD5);
@@ -367,6 +373,8 @@ void FDatasmithUEPbrMaterialElementImpl::ResetExpressionGraph( bool bRemoveAllEx
 	WorldDisplacement.Edit()->SetExpression(nullptr);
 	Refraction.Edit()->SetExpression(nullptr);
 	AmbientOcclusion.Edit()->SetExpression(nullptr);
+	ClearCoat.Edit()->SetExpression(nullptr);
+	ClearCoatRoughness.Edit()->SetExpression(nullptr);
 	MaterialAttributes.Edit()->SetExpression(nullptr);
 }
 
diff --git a/Engine/Source/Runtime/Datasmith/DatasmithCore/Private/DatasmithMaterialElementsImpl.h b/Engine/Source/Runtime/Datasmith/DatasmithCore/Private/DatasmithMaterialElementsImpl.h
--- a/Engine/Source/Runtime/Datasmith/DatasmithCore/Private/DatasmithMaterialElementsImpl.h
+++ b/Engine/Source/Runtime/Datasmith/DatasmithCore/Private/DatasmithMaterialElementsImpl.h
@@ -674,6 +674,8 @@ public:
 	virtual IDatasmithExpressionInput& GetNormal() override { return *Normal.Edit(); }
 	virtual IDatasmithExpressionInput& GetRefraction() override { return *Refraction.Edit(); }
 	virtual IDatasmithExpressionInput& GetAmbientOcclusion() override { return *AmbientOcclusion.Edit(); }
+	virtual IDatasmithExpressionInput& GetClearCoat() override { return *ClearCoat.Edit(); }
+	virtual IDatasmithExpressionInput& GetClearCoatRoughness() override { return *ClearCoatRoughness.Edit(); }
 	virtual IDatasmithExpressionInput& GetMaterialAttributes() override { return *MaterialAttributes.Edit(); }
 
 	virtual int GetBlendMode() const override {return BlendMode; }
@@ -721,6 +723,8 @@ protected:
 	TDatasmithReferenceProxy< FDatasmithExpressionInputImpl > WorldDisplacement;
 	TDatasmithReferenceProxy< FDatasmithExpressionInputImpl > Refraction;
 	TDatasmithReferenceProxy< FDatasmithExpressionInputImpl > AmbientOcclusion;
+	TDatasmithReferenceProxy< FDatasmithExpressionInputImpl > ClearCoat;
+	TDatasmithReferenceProxy< FDatasmithExpressionInputImpl > ClearCoatRoughness;
 	TDatasmithReferenceProxy< FDatasmithExpressionInputImpl > MaterialAttributes;
 
 	TDatasmithReferenceArrayProxy< IDatasmithMaterialExpression > Expressions;
@@ -730,7 +734,7 @@ protected:
 	TReflected<bool> bUseMaterialAttributes;
 	TReflected<bool> bMaterialFunctionOnly;
 	TReflected<float> OpacityMaskClipValue;
-	TReflected<int> TranslucencyLightingMode;
+	TReflected<int32> TranslucencyLightingMode;
 
 	TReflected<FString> ParentLabel;
 	TReflected<EDatasmithShadingModel, uint8> ShadingModel;
diff --git a/Engine/Source/Runtime/Datasmith/DatasmithCore/Public/DatasmithMaterialElements.h b/Engine/Source/Runtime/Datasmith/DatasmithCore/Public/DatasmithMaterialElements.h
--- a/Engine/Source/Runtime/Datasmith/DatasmithCore/Public/DatasmithMaterialElements.h
+++ b/Engine/Source/Runtime/Datasmith/DatasmithCore/Public/DatasmithMaterialElements.h
@@ -221,6 +221,8 @@ public:
 	virtual IDatasmithExpressionInput& GetNormal() = 0;
 	virtual IDatasmithExpressionInput& GetRefraction() = 0;
 	virtual IDatasmithExpressionInput& GetAmbientOcclusion() = 0;
+	virtual IDatasmithExpressionInput& GetClearCoat() = 0;
+	virtual IDatasmithExpressionInput& GetClearCoatRoughness() = 0;
 	virtual IDatasmithExpressionInput& GetMaterialAttributes() = 0;
 
 	virtual int GetBlendMode() const = 0;
