Datasmith glTF: Add support for transmission extension

diff --git a/Engine/Plugins/Editor/GLTFImporter/Source/GLTFCore/Private/GLTFMaterialFactory.cpp b/Engine/Plugins/Editor/GLTFImporter/Source/GLTFCore/Private/GLTFMaterialFactory.cpp
--- a/Engine/Plugins/Editor/GLTFImporter/Source/GLTFCore/Private/GLTFMaterialFactory.cpp
+++ b/Engine/Plugins/Editor/GLTFImporter/Source/GLTFCore/Private/GLTFMaterialFactory.cpp
@@ -448,6 +448,36 @@ namespace GLTF
 			TEXT("Transmission"),
 			ETextureMode::Color);
 
+		FMaterialExpression* MetallicFactor = MaterialElement.GetMetallic().GetExpression();
+		FMaterialExpression* RoughnessFactor = MaterialElement.GetRoughness().GetExpression();
+
+		FMaterialExpression* MetallicExpr = nullptr;
+		FMaterialExpression* RoughnessExpr = nullptr;
+
+		FMaterialExpressionTexture* MetallicRoughnessTexture = FindExpression<FMaterialExpressionTexture>(TEXT("MetallicRoughness"), MaterialElement);
+		
+		if (MetallicRoughnessTexture)
+		{
+			FMaterialExpressionGeneric* MultMetallic = MaterialElement.AddMaterialExpression<FMaterialExpressionGeneric>();
+			MultMetallic->SetExpressionName(TEXT("Multiply"));
+			FMaterialExpressionGeneric* MultRoughness = MaterialElement.AddMaterialExpression<FMaterialExpressionGeneric>();
+			MultRoughness->SetExpressionName(TEXT("Multiply"));
+
+			MetallicRoughnessTexture->ConnectExpression(*MultMetallic->GetInput(0), (int)FPBRMapFactory::EChannel::Blue);
+			MetallicFactor->ConnectExpression(*MultMetallic->GetInput(1), 0);
+
+			MetallicRoughnessTexture->ConnectExpression(*MultRoughness->GetInput(0), (int)FPBRMapFactory::EChannel::Green);
+			RoughnessFactor->ConnectExpression(*MultMetallic->GetInput(1), 0);
+
+			MetallicExpr = MultMetallic;
+			RoughnessExpr = MultRoughness;
+		}
+		else
+		{
+			MetallicExpr = MetallicFactor;
+			RoughnessExpr = RoughnessFactor;
+		}
+
 		FMaterialExpressionTexture* BaseColorMap = FindExpression<FMaterialExpressionTexture>(TEXT("BaseColor Map"), MaterialElement);
 
 		if (GLTFMaterial.AlphaMode == FMaterial::EAlphaMode::Mask)
@@ -498,60 +528,65 @@ namespace GLTF
 			Step1->ConnectExpression(*OneMinus1->GetInput(0), 0);
 			Step1->ConnectExpression(MaterialElement.GetSpecular(), 0);
 
-			FMaterialExpressionGeneric* Lerp1 = MaterialElement.AddMaterialExpression<FMaterialExpressionGeneric>();
-			Lerp1->SetExpressionName(TEXT("LinearInterpolate"));
-			OneMinus1->ConnectExpression(*Lerp1->GetInput(2), 0);
+			FMaterialExpressionGeneric* OpacityMult = MaterialElement.AddMaterialExpression<FMaterialExpressionGeneric>();
+			OpacityMult->SetExpressionName(TEXT("Multiply"));
 
-			FMaterialExpressionGeneric* Mult2 = MaterialElement.AddMaterialExpression<FMaterialExpressionGeneric>();
-			Mult2->SetExpressionName(TEXT("Multiply"));
+			OneMinus1->ConnectExpression(*OpacityMult->GetInput(0), 0);
+			OneMinus1->ConnectExpression(MaterialElement.GetSpecular(), 0);
 
-			Mult2->ConnectExpression(MaterialElement.GetOpacity(), 0);
+			OpacityMult->ConnectExpression(MaterialElement.GetOpacity(), 0);
 
-			if (TransmissionTexture)
-			{
-				FMaterialExpressionGeneric* Mult3 = MaterialElement.AddMaterialExpression<FMaterialExpressionGeneric>();
-				Mult3->SetExpressionName(TEXT("Multiply"));
+			FMaterialExpressionGeneric* Max1 = MaterialElement.AddMaterialExpression<FMaterialExpressionGeneric>();
+			Max1->SetExpressionName(TEXT("Max"));
 
-				FMaterialExpressionGeneric* Mult4 = MaterialElement.AddMaterialExpression<FMaterialExpressionGeneric>();
-				Mult4->SetExpressionName(TEXT("Multiply"));
+			FMaterialExpressionGeneric* Max2 = MaterialElement.AddMaterialExpression<FMaterialExpressionGeneric>();
+			Max2->SetExpressionName(TEXT("Max"));
 
-				FMaterialExpressionGeneric* OneMinus2 = MaterialElement.AddMaterialExpression<FMaterialExpressionGeneric>();
-				OneMinus2->SetExpressionName(TEXT("OneMinus"));
+			MetallicExpr->ConnectExpression(*Max1->GetInput(0), 0);
+			RoughnessExpr->ConnectExpression(*Max1->GetInput(1), 0);
+	
+			Max1->ConnectExpression(*Max2->GetInput(1), 0);
+			Max2->ConnectExpression(*OpacityMult->GetInput(1), 0);
 
-				Mult1->ConnectExpression(*Mult3->GetInput(0), 0);
+			FMaterialExpression* TransmissionExpr = nullptr;
 
-				Mult3->ConnectExpression(*Mult4->GetInput(0), 0);
-				Mult3->ConnectExpression(*OneMinus2->GetInput(0), 0);
-				Mult1->ConnectExpression(*Mult4->GetInput(1), 0);
-				Mult4->ConnectExpression(*Lerp1->GetInput(0), 0);
+			if (TransmissionTexture)
+			{
+				FMaterialExpressionGeneric* MultTransm = MaterialElement.AddMaterialExpression<FMaterialExpressionGeneric>();
+				MultTransm->SetExpressionName(TEXT("Multiply"));
 
-				Step1->ConnectExpression(*Mult2->GetInput(0), 0);
-				OneMinus2->ConnectExpression(*Mult2->GetInput(1), 0);
-				
-				TransmissionTexture->ConnectExpression(*Mult3->GetInput(0), 1);
-				TransmissionFactorExpr->ConnectExpression(*Mult3->GetInput(1), 0);
-				Lerp1->ConnectExpression(*ThinTranslucentOutput->GetInput(0), 0);
+				TransmissionTexture->ConnectExpression(*MultTransm->GetInput(0), (int)FPBRMapFactory::EChannel::Red);
+				TransmissionFactorExpr->ConnectExpression(*MultTransm->GetInput(1), 0);
+
+				TransmissionExpr = MultTransm;
 			}
 			else
 			{
-				FMaterialExpressionScalar* Const1 = MaterialElement.AddMaterialExpression<FMaterialExpressionScalar>();
-				Const1->GetScalar() = 0;
+				TransmissionExpr = TransmissionFactorExpr;
+			}
 
-				Step1->ConnectExpression(*Mult2->GetInput(1), 0);
-				Const1->ConnectExpression(*Mult2->GetInput(0), 0);
+			// Transmission network
+			FMaterialExpressionGeneric* Mult2 = MaterialElement.AddMaterialExpression<FMaterialExpressionGeneric>();
+			Mult2->SetExpressionName(TEXT("Multiply"));
 
-				Mult1->ConnectExpression(*Lerp1->GetInput(0), 0);
+			FMaterialExpressionGeneric* OneMinus2 = MaterialElement.AddMaterialExpression<FMaterialExpressionGeneric>();
+			OneMinus2->SetExpressionName(TEXT("OneMinus"));
 
-				FMaterialExpressionGeneric* Mult3 = MaterialElement.AddMaterialExpression<FMaterialExpressionGeneric>();
-				Mult3->SetExpressionName(TEXT("Multiply"));
+			FMaterialExpressionGeneric* LerpThinTransl = MaterialElement.AddMaterialExpression<FMaterialExpressionGeneric>();
+			LerpThinTransl->SetExpressionName(TEXT("LinearInterpolate"));
 
-				Lerp1->ConnectExpression(*Mult3->GetInput(0), 0);
-				TransmissionFactorExpr->ConnectExpression(*Mult3->GetInput(1), 0);
+			TransmissionExpr->ConnectExpression(*Mult2->GetInput(0), 0);
+			TransmissionExpr->ConnectExpression(*OneMinus2->GetInput(0), 0);
+			Mult1->ConnectExpression(*Mult2->GetInput(1), 0);
 
-				Mult3->ConnectExpression(*ThinTranslucentOutput->GetInput(0), 0);
-			}
+			Mult2->ConnectExpression(*LerpThinTransl->GetInput(0), 0);
+			Step1->ConnectExpression(*LerpThinTransl->GetInput(2), 0);
+
+			OneMinus2->ConnectExpression(*Max2->GetInput(0), 0);
 
-			BaseColor2->ConnectExpression(MaterialElement.GetBaseColor(), 0);
+			LerpThinTransl->ConnectExpression(*ThinTranslucentOutput->GetInput(0), 0);
+		
+			Mult1->ConnectExpression(MaterialElement.GetBaseColor(), 0);
 		}
 		else
 		{
@@ -560,8 +595,19 @@ namespace GLTF
 
 			BaseColorMult->ConnectExpression(*ThinTranslucentOutput->GetInput(0), 0);
 
-			FMaterialExpressionGeneric* OpacitySubtract = MaterialElement.AddMaterialExpression<FMaterialExpressionGeneric>();
-			OpacitySubtract->SetExpressionName(TEXT("Subtract"));
+			FMaterialExpressionGeneric* OneMinus = MaterialElement.AddMaterialExpression<FMaterialExpressionGeneric>();
+			OneMinus->SetExpressionName(TEXT("OneMinus"));
+
+			FMaterialExpressionGeneric* Max1 = MaterialElement.AddMaterialExpression<FMaterialExpressionGeneric>();
+			Max1->SetExpressionName(TEXT("Max"));
+
+			FMaterialExpressionGeneric* Max2 = MaterialElement.AddMaterialExpression<FMaterialExpressionGeneric>();
+			Max2->SetExpressionName(TEXT("Max"));
+
+			MetallicExpr->ConnectExpression(*Max1->GetInput(0), 0);
+			RoughnessExpr->ConnectExpression(*Max1->GetInput(1), 0);
+			OneMinus->ConnectExpression(*Max2->GetInput(0), 0);
+			Max1->ConnectExpression(*Max2->GetInput(1), 0);
 
 			FMaterialExpression* BaseColorExpr = MaterialElement.GetBaseColor().GetExpression();
 			if (BaseColorExpr)
@@ -576,15 +622,15 @@ namespace GLTF
 				TransmissionFactorExpr->ConnectExpression(*TransmissionMult->GetInput(0), 0);
 				TransmissionTexture->ConnectExpression(*TransmissionMult->GetInput(1), (int)FPBRMapFactory::EChannel::Red);
 				TransmissionMult->ConnectExpression(*BaseColorMult->GetInput(1), 0);
-				TransmissionMult->ConnectExpression(*OpacitySubtract->GetInput(1), 0);
+				TransmissionMult->ConnectExpression(*OneMinus->GetInput(0), 0);
 			}
 			else
 			{
 				TransmissionFactorExpr->ConnectExpression(*BaseColorMult->GetInput(1), 0);
-				TransmissionFactorExpr->ConnectExpression(*OpacitySubtract->GetInput(1), 0);
+				TransmissionFactorExpr->ConnectExpression(*OneMinus->GetInput(0), 0);
 			}
 
-			OpacitySubtract->ConnectExpression(MaterialElement.GetOpacity(), 0);
+			Max2->ConnectExpression(MaterialElement.GetOpacity(), 0);
 		}
 	}
 	
