Scale of gltf imported texture are wrong

diff --git a/Engine/Plugins/Editor/GLTFImporter/Source/GLTFCore/Private/GLTF/ExtensionsHandler.cpp b/Engine/Plugins/Editor/GLTFImporter/Source/GLTFCore/Private/GLTF/ExtensionsHandler.cpp
--- a/Engine/Plugins/Editor/GLTFImporter/Source/GLTFCore/Private/GLTF/ExtensionsHandler.cpp
+++ b/Engine/Plugins/Editor/GLTFImporter/Source/GLTFCore/Private/GLTF/ExtensionsHandler.cpp
@@ -136,10 +136,10 @@ namespace GLTF
 				case KHR_materials_pbrSpecularGlossiness:
 				{
 					const FJsonObject& PBR = ExtObj;
-					GLTF::SetTextureMap(PBR, TEXT("diffuseTexture"), nullptr, Asset->Textures, Material.BaseColor);
+					GLTF::SetTextureMap(PBR, TEXT("diffuseTexture"), nullptr, Asset->Textures, Material.BaseColor, Messages);
 					Material.BaseColorFactor = GetVec4(PBR, TEXT("diffuseFactor"), FVector4(1.0f, 1.0f, 1.0f, 1.0f));
 
-					GLTF::SetTextureMap(PBR, TEXT("specularGlossinessTexture"), nullptr, Asset->Textures, Material.SpecularGlossiness.Map);
+					GLTF::SetTextureMap(PBR, TEXT("specularGlossinessTexture"), nullptr, Asset->Textures, Material.SpecularGlossiness.Map, Messages);
 					Material.SpecularGlossiness.SpecularFactor   = GetVec3(PBR, TEXT("specularFactor"), FVector(1.0f));
 					Material.SpecularGlossiness.GlossinessFactor = GetScalar(PBR, TEXT("glossinessFactor"), 1.0f);
 
@@ -161,12 +161,12 @@ namespace GLTF
 					Material.bHasClearCoat = true;
 
 					Material.ClearCoat.ClearCoatFactor = GetScalar(ClearCoat, TEXT("clearcoatFactor"), 1.0f);
-					GLTF::SetTextureMap(ClearCoat, TEXT("clearcoatTexture"), nullptr, Asset->Textures, Material.ClearCoat.ClearCoatMap);
+					GLTF::SetTextureMap(ClearCoat, TEXT("clearcoatTexture"), nullptr, Asset->Textures, Material.ClearCoat.ClearCoatMap, Messages);
 
 					Material.ClearCoat.Roughness = GetScalar(ClearCoat, TEXT("clearcoatRoughnessFactor"), 0.0f);
-					GLTF::SetTextureMap(ClearCoat, TEXT("clearcoatRoughnessTexture"), nullptr, Asset->Textures, Material.ClearCoat.RoughnessMap);
+					GLTF::SetTextureMap(ClearCoat, TEXT("clearcoatRoughnessTexture"), nullptr, Asset->Textures, Material.ClearCoat.RoughnessMap, Messages);
 
-					Material.ClearCoat.NormalMapUVScale = GLTF::SetTextureMap(ClearCoat, TEXT("clearcoatNormalTexture"), TEXT("scale"), Asset->Textures, Material.ClearCoat.NormalMap);
+					Material.ClearCoat.NormalMapUVScale = GLTF::SetTextureMap(ClearCoat, TEXT("clearcoatNormalTexture"), TEXT("scale"), Asset->Textures, Material.ClearCoat.NormalMap, Messages);
 
 					Asset->ExtensionsUsed.Add(EExtension::KHR_MaterialsClearCoat);
 				}
@@ -178,7 +178,7 @@ namespace GLTF
 					Material.bHasTransmission = true;
 
 					Material.Transmission.TransmissionFactor = GetScalar(Transm, TEXT("transmissionFactor"), 0.0f);
-					GLTF::SetTextureMap(Transm, TEXT("transmissionTexture"), nullptr, Asset->Textures, Material.Transmission.TransmissionMap);
+					GLTF::SetTextureMap(Transm, TEXT("transmissionTexture"), nullptr, Asset->Textures, Material.Transmission.TransmissionMap, Messages);
 
 					Asset->ExtensionsUsed.Add(EExtension::KHR_MaterialsTransmission);
 				}
@@ -210,8 +210,8 @@ namespace GLTF
 					Material.bHasSpecular = true;
 					Material.Specular.SpecularFactor = GetScalar(Specular, TEXT("specularFactor"), 1.0f);
 					Material.Specular.SpecularColorFactor = GetVec3(Specular, TEXT("specularColorFactor"));
-					GLTF::SetTextureMap(Specular, TEXT("specularTexture"), nullptr, Asset->Textures, Material.Specular.SpecularMap);
-					GLTF::SetTextureMap(Specular, TEXT("specularColorTexture"), nullptr, Asset->Textures, Material.Specular.SpecularColorMap);
+					GLTF::SetTextureMap(Specular, TEXT("specularTexture"), nullptr, Asset->Textures, Material.Specular.SpecularMap, Messages);
+					GLTF::SetTextureMap(Specular, TEXT("specularColorTexture"), nullptr, Asset->Textures, Material.Specular.SpecularColorMap, Messages);
 
 					Asset->ExtensionsUsed.Add(EExtension::KHR_MaterialsSpecular);
 				}
@@ -219,15 +219,15 @@ namespace GLTF
 				case MSFT_packing_occlusionRoughnessMetallic:
 				{
 					const FJsonObject& Packing = ExtObj;
-					if (GLTF::SetTextureMap(Packing, TEXT("occlusionRoughnessMetallicTexture"), nullptr, Asset->Textures, Material.Packing.Map))
+					if (GLTF::SetTextureMap(Packing, TEXT("occlusionRoughnessMetallicTexture"), nullptr, Asset->Textures, Material.Packing.Map, Messages))
 					{
 						Material.Packing.Flags = (int)GLTF::FMaterial::EPackingFlags::OcclusionRoughnessMetallic;
 					}
-					else if (GLTF::SetTextureMap(Packing, TEXT("roughnessMetallicOcclusionTexture"), nullptr, Asset->Textures, Material.Packing.Map))
+					else if (GLTF::SetTextureMap(Packing, TEXT("roughnessMetallicOcclusionTexture"), nullptr, Asset->Textures, Material.Packing.Map, Messages))
 					{
 						Material.Packing.Flags = (int)GLTF::FMaterial::EPackingFlags::RoughnessMetallicOcclusion;
 					}
-					if (GLTF::SetTextureMap(Packing, TEXT("normalTexture"), nullptr, Asset->Textures, Material.Packing.NormalMap))
+					if (GLTF::SetTextureMap(Packing, TEXT("normalTexture"), nullptr, Asset->Textures, Material.Packing.NormalMap, Messages))
 					{
 						// can have extra packed two channel(RG) normal map
 						Material.Packing.Flags |= (int)GLTF::FMaterial::EPackingFlags::NormalRG;
@@ -240,7 +240,7 @@ namespace GLTF
 				case MSFT_packing_normalRoughnessMetallic:
 				{
 					const FJsonObject& Packing = ExtObj;
-					if (GLTF::SetTextureMap(Packing, TEXT("normalRoughnessMetallicTexture"), nullptr, Asset->Textures, Material.Packing.Map))
+					if (GLTF::SetTextureMap(Packing, TEXT("normalRoughnessMetallicTexture"), nullptr, Asset->Textures, Material.Packing.Map, Messages))
 					{
 						Material.Packing.NormalMap = Material.Packing.Map;
 						Material.Packing.Flags     = (int)GLTF::FMaterial::EPackingFlags::NormalRoughnessMetallic;
diff --git a/Engine/Plugins/Editor/GLTFImporter/Source/GLTFCore/Private/GLTF/ExtensionsHandler.h b/Engine/Plugins/Editor/GLTFImporter/Source/GLTFCore/Private/GLTF/ExtensionsHandler.h
--- a/Engine/Plugins/Editor/GLTFImporter/Source/GLTFCore/Private/GLTF/ExtensionsHandler.h
+++ b/Engine/Plugins/Editor/GLTFImporter/Source/GLTFCore/Private/GLTF/ExtensionsHandler.h
@@ -47,6 +47,11 @@ namespace GLTF
 		void SetupTextureExtensions(const FJsonObject& Object, FTexture& Texture) const;
 		void SetupMaterialExtensions(const FJsonObject& Object, FMaterial& Material) const;
 
+		TArray<FLogMessage>& GetMessages()
+		{
+			return Messages;
+		}
+
 	private:
 		void CheckExtensions(const FJsonObject& Object, const TArray<FString>& ExtensionsSupported) const;
 
diff --git a/Engine/Plugins/Editor/GLTFImporter/Source/GLTFCore/Private/GLTF/GLTFReader.cpp b/Engine/Plugins/Editor/GLTFImporter/Source/GLTFCore/Private/GLTF/GLTFReader.cpp
--- a/Engine/Plugins/Editor/GLTFImporter/Source/GLTFCore/Private/GLTF/GLTFReader.cpp
+++ b/Engine/Plugins/Editor/GLTFImporter/Source/GLTFCore/Private/GLTF/GLTFReader.cpp
@@ -569,20 +569,20 @@ namespace GLTF
 		Asset->Materials.Emplace(GetString(Object, TEXT("name")));
 		FMaterial& Material = Asset->Materials.Last();
 
-		GLTF::SetTextureMap(Object, TEXT("emissiveTexture"), nullptr, Asset->Textures, Material.Emissive);
+		GLTF::SetTextureMap(Object, TEXT("emissiveTexture"), nullptr, Asset->Textures, Material.Emissive, ExtensionsHandler->GetMessages());
 		Material.EmissiveFactor = GetVec3(Object, TEXT("emissiveFactor"));
 
-		Material.NormalScale       = GLTF::SetTextureMap(Object, TEXT("normalTexture"), TEXT("scale"), Asset->Textures, Material.Normal);
-		Material.OcclusionStrength = GLTF::SetTextureMap(Object, TEXT("occlusionTexture"), TEXT("strength"), Asset->Textures, Material.Occlusion);
+		Material.NormalScale       = GLTF::SetTextureMap(Object, TEXT("normalTexture"), TEXT("scale"), Asset->Textures, Material.Normal, ExtensionsHandler->GetMessages());
+		Material.OcclusionStrength = GLTF::SetTextureMap(Object, TEXT("occlusionTexture"), TEXT("strength"), Asset->Textures, Material.Occlusion, ExtensionsHandler->GetMessages());
 
 		if (Object.HasTypedField<EJson::Object>(TEXT("pbrMetallicRoughness")))
 		{
 			const FJsonObject& PBR = *Object.GetObjectField(TEXT("pbrMetallicRoughness"));
 
-			GLTF::SetTextureMap(PBR, TEXT("baseColorTexture"), nullptr, Asset->Textures, Material.BaseColor);
+			GLTF::SetTextureMap(PBR, TEXT("baseColorTexture"), nullptr, Asset->Textures, Material.BaseColor, ExtensionsHandler->GetMessages());
 			Material.BaseColorFactor = GetVec4(PBR, TEXT("baseColorFactor"), FVector4(1.0f, 1.0f, 1.0f, 1.0f));
 
-			GLTF::SetTextureMap(PBR, TEXT("metallicRoughnessTexture"), nullptr, Asset->Textures, Material.MetallicRoughness.Map);
+			GLTF::SetTextureMap(PBR, TEXT("metallicRoughnessTexture"), nullptr, Asset->Textures, Material.MetallicRoughness.Map, ExtensionsHandler->GetMessages());
 			Material.MetallicRoughness.MetallicFactor  = GetScalar(PBR, TEXT("metallicFactor"), 1.0f);
 			Material.MetallicRoughness.RoughnessFactor = GetScalar(PBR, TEXT("roughnessFactor"), 1.0f);
 		}
diff --git a/Engine/Plugins/Editor/GLTFImporter/Source/GLTFCore/Private/GLTF/MaterialUtilities.h b/Engine/Plugins/Editor/GLTFImporter/Source/GLTFCore/Private/GLTF/MaterialUtilities.h
--- a/Engine/Plugins/Editor/GLTFImporter/Source/GLTFCore/Private/GLTF/MaterialUtilities.h
+++ b/Engine/Plugins/Editor/GLTFImporter/Source/GLTFCore/Private/GLTF/MaterialUtilities.h
@@ -12,7 +12,7 @@ namespace GLTF
 {
 	// Returns scale factor if JSON has it, 1.0 by default.
 	inline float SetTextureMap(const FJsonObject& InObject, const TCHAR* InTexName, const TCHAR* InScaleName, const TArray<FTexture>& Textures,
-	                           GLTF::FTextureMap& OutMap)
+	                           GLTF::FTextureMap& OutMap, TArray<FLogMessage>& OutMessages)
 	{
 		float Scale = 1.0f;
 
@@ -30,6 +30,64 @@ namespace GLTF
 				{
 					Scale = TexObj.GetNumberField(InScaleName);
 				}
+
+				// Handle textureInfo extensions
+				if (TexObj.HasTypedField<EJson::Object>(TEXT("extensions")))
+				{
+					enum
+					{
+						KHR_texture_transform = 0
+					};
+
+					static const TArray<FString> SupportedTextureInfoExtensions = { TEXT("KHR_texture_transform") };
+
+					const FJsonObject& ExtensionsObj = *TexObj.GetObjectField(TEXT("extensions"));
+
+					for (int32 Index = 0; Index < SupportedTextureInfoExtensions.Num(); ++Index)
+					{
+						const FString ExtensionName = SupportedTextureInfoExtensions[Index];
+						if (!ExtensionsObj.HasTypedField<EJson::Object>(ExtensionName))
+						{
+							OutMessages.Emplace(EMessageSeverity::Warning, FString::Printf(TEXT("Extension is not supported: %s"), *ExtensionName));
+							continue;
+						}
+
+						const FJsonObject& ExtObj = *ExtensionsObj.GetObjectField(ExtensionName);
+						switch (Index)
+						{
+							case KHR_texture_transform:
+							{
+								OutMap.bHasTextureTransform = true;
+								OutMap.TextureTransform.Offset[0] = 0.0f;
+								OutMap.TextureTransform.Offset[1] = 0.0f;
+								OutMap.TextureTransform.Scale[0] = 1.0f;
+								OutMap.TextureTransform.Scale[1] = 1.0f;
+								OutMap.TextureTransform.Rotation = GetScalar(ExtObj, TEXT("rotation"), 0.f);
+
+								if (ExtObj.HasTypedField<EJson::Array>(TEXT("offset")))
+								{
+									const TArray<TSharedPtr<FJsonValue>>& Offset = ExtObj.GetArrayField(TEXT("offset"));
+									if (Offset.Num() >= 2)
+									{
+										OutMap.TextureTransform.Offset[0] = static_cast<float>(Offset[0]->AsNumber());
+										OutMap.TextureTransform.Offset[1] = static_cast<float>(Offset[1]->AsNumber());
+									}
+								}
+
+								if (ExtObj.HasTypedField<EJson::Array>(TEXT("scale")))
+								{
+									const TArray<TSharedPtr<FJsonValue>>& TransformScale = ExtObj.GetArrayField(TEXT("scale"));
+									if(TransformScale.Num() >= 2)
+									{
+										OutMap.TextureTransform.Scale[0] = static_cast<float>(TransformScale[0]->AsNumber());
+										OutMap.TextureTransform.Scale[1] = static_cast<float>(TransformScale[1]->AsNumber());
+									}
+								}
+							}
+							break;
+						}
+					}
+				}
 			}
 		}
 		else
diff --git a/Engine/Plugins/Editor/GLTFImporter/Source/GLTFCore/Private/GLTFMapFactory.cpp b/Engine/Plugins/Editor/GLTFImporter/Source/GLTFCore/Private/GLTFMapFactory.cpp
--- a/Engine/Plugins/Editor/GLTFImporter/Source/GLTFCore/Private/GLTFMapFactory.cpp
+++ b/Engine/Plugins/Editor/GLTFImporter/Source/GLTFCore/Private/GLTFMapFactory.cpp
@@ -21,6 +21,71 @@ namespace GLTF
 			}
 		}
 
+		void CreateTextureTransform(const FTextureTransform& InTextureTransform, FMaterialExpressionTexture& TexExpression, FMaterialElement& MaterialElement)
+		{
+			FMaterialExpression* CoordExpression = TexExpression.GetInputCoordinate().GetExpression();
+
+			if (CoordExpression == nullptr)
+			{
+				CoordExpression = MaterialElement.AddMaterialExpression<FMaterialExpressionTextureCoordinate>();
+			}
+
+			FMaterialExpression* LastExprInChain = nullptr;
+
+			FMaterialExpressionGeneric* UVMultiplyExpr = MaterialElement.AddMaterialExpression<FMaterialExpressionGeneric>();
+			UVMultiplyExpr->SetExpressionName(TEXT("Multiply"));
+
+			FMaterialExpressionGeneric* UVScale = MaterialElement.AddMaterialExpression<FMaterialExpressionGeneric>();
+			UVScale->SetExpressionName(TEXT("Constant2Vector"));
+			UVScale->SetFloatProperty(TEXT("R"), InTextureTransform.Scale[0]);
+			UVScale->SetFloatProperty(TEXT("G"), InTextureTransform.Scale[1]);
+
+			CoordExpression->ConnectExpression(*UVMultiplyExpr->GetInput(0), 0);
+			UVScale->ConnectExpression(*UVMultiplyExpr->GetInput(1), 0);
+
+			LastExprInChain = UVMultiplyExpr;
+
+			if (!FMath::IsNearlyZero(InTextureTransform.Rotation))
+			{
+				FMaterialExpressionFunctionCall* UVRotate = MaterialElement.AddMaterialExpression<FMaterialExpressionFunctionCall>();
+				UVRotate->SetFunctionPathName(TEXT("/Engine/Functions/Engine_MaterialFunctions02/Texturing/CustomRotator.CustomRotator"));
+
+				float AngleRadians = InTextureTransform.Rotation;
+
+				if (AngleRadians < 0.0f)
+				{
+					AngleRadians = TWO_PI - AngleRadians;
+				}
+
+				FMaterialExpressionScalar* RotationAngle = MaterialElement.AddMaterialExpression<FMaterialExpressionScalar>();
+				RotationAngle->GetScalar() = AngleRadians / TWO_PI; // Normalize angle value
+
+				LastExprInChain->ConnectExpression(*UVRotate->GetInput(0), 0);
+				RotationAngle->ConnectExpression(*UVRotate->GetInput(2), 0);
+
+				LastExprInChain = UVRotate;
+			}
+
+			if (!FMath::IsNearlyZero(InTextureTransform.Offset[0]) || 
+				!FMath::IsNearlyZero(InTextureTransform.Offset[1]))
+			{
+				FMaterialExpressionGeneric* UVOffset = MaterialElement.AddMaterialExpression<FMaterialExpressionGeneric>();
+				UVOffset->SetExpressionName(TEXT("Constant2Vector"));
+				UVOffset->SetFloatProperty(TEXT("R"), InTextureTransform.Offset[0]);
+				UVOffset->SetFloatProperty(TEXT("G"), InTextureTransform.Offset[1]);
+
+				FMaterialExpressionGeneric* AddExpr = MaterialElement.AddMaterialExpression<FMaterialExpressionGeneric>();
+				AddExpr->SetExpressionName(TEXT("Add"));
+
+				LastExprInChain->ConnectExpression(*AddExpr->GetInput(0), 0);
+				UVOffset->ConnectExpression(*AddExpr->GetInput(1), 0);
+
+				LastExprInChain = AddExpr;
+			}
+
+			LastExprInChain->ConnectExpression(TexExpression.GetInputCoordinate(), 0);
+		}
+
 		FMaterialExpressionInput& GetFirstInput(FMaterialExpression* Expression)
 		{
 			FMaterialExpressionInput* Input = Expression->GetInput(0);
@@ -54,11 +119,11 @@ namespace GLTF
 	{
 	}
 
-	void FPBRMapFactory::CreateNormalMap(const GLTF::FTexture& Map, int CoordinateIndex, float NormalScale)
+	void FPBRMapFactory::CreateNormalMap(const GLTF::FTexture& Map, int CoordinateIndex, float NormalScale, const GLTF::FTextureTransform* TextureTransform)
 	{
 		check(CurrentMaterialElement);
 
-		FMaterialExpressionTexture* TexExpression = CreateTextureMap(Map, CoordinateIndex, TEXT("Normal Map"), ETextureMode::Normal);
+		FMaterialExpressionTexture* TexExpression = CreateTextureMap(Map, CoordinateIndex, TEXT("Normal Map"), ETextureMode::Normal, TextureTransform);
 		if (!TexExpression)
 		{
 			return;
@@ -104,19 +169,21 @@ namespace GLTF
 	}
 
 	FMaterialExpression* FPBRMapFactory::CreateColorMap(const GLTF::FTexture& Map, int CoordinateIndex, const FVector3f& Color, const TCHAR* MapName,
-	                                                    const TCHAR* ValueName, ETextureMode TextureMode, FMaterialExpressionInput& MaterialInput)
+	                                                    const TCHAR* ValueName, ETextureMode TextureMode, FMaterialExpressionInput& MaterialInput,
+														const GLTF::FTextureTransform* TextureTransform)
 	{
-		return CreateMap<FMaterialExpressionColor>(Map, CoordinateIndex, Color, MapName, ValueName, TextureMode, MaterialInput);
+		return CreateMap<FMaterialExpressionColor>(Map, CoordinateIndex, Color, MapName, ValueName, TextureMode, MaterialInput, TextureTransform);
 	}
 
 	FMaterialExpression* FPBRMapFactory::CreateColorMap(const GLTF::FTexture& Map, int CoordinateIndex, const FVector4f& Color, const TCHAR* MapName,
-	                                                    const TCHAR* ValueName, ETextureMode TextureMode, FMaterialExpressionInput& MaterialInput)
+	                                                    const TCHAR* ValueName, ETextureMode TextureMode, FMaterialExpressionInput& MaterialInput, 
+														const GLTF::FTextureTransform* TextureTransform)
 	{
-		return CreateMap<FMaterialExpressionColor>(Map, CoordinateIndex, Color, MapName, ValueName, TextureMode, MaterialInput);
+		return CreateMap<FMaterialExpressionColor>(Map, CoordinateIndex, Color, MapName, ValueName, TextureMode, MaterialInput, TextureTransform);
 	}
 
 	void FPBRMapFactory::CreateMultiMap(const GLTF::FTexture& Map, int CoordinateIndex, const TCHAR* MapName, const FMapChannel* MapChannels,
-	                                    uint32 MapChannelsCount, ETextureMode TextureMode)
+	                                    uint32 MapChannelsCount, ETextureMode TextureMode, const GLTF::FTextureTransform* TextureTransform)
 	{
 		FExpressionList ValueExpressions;
 		for (uint32 Index = 0; Index < MapChannelsCount; ++Index)
@@ -149,7 +216,7 @@ namespace GLTF
 			ValueExpressions.Add(ValueExpression);
 		}
 
-		FMaterialExpressionTexture* TexExpression = CreateTextureMap(Map, CoordinateIndex, MapName, TextureMode);
+		FMaterialExpressionTexture* TexExpression = CreateTextureMap(Map, CoordinateIndex, MapName, TextureMode, TextureTransform);
 		if (TexExpression)
 		{
 			for (uint32 Index = 0; Index < MapChannelsCount; ++Index)
@@ -217,7 +284,7 @@ namespace GLTF
 	}
 
 	FMaterialExpressionTexture* FPBRMapFactory::CreateTextureMap(const GLTF::FTexture& Map, int CoordinateIndex, const TCHAR* MapName,
-	                                                             ETextureMode TextureMode)
+	                                                             ETextureMode TextureMode, const GLTF::FTextureTransform* TextureTransform)
 	{
 		ITextureElement* Texture = TextureFactory.CreateTexture(Map, ParentPackage, Flags, TextureMode);
 
@@ -230,13 +297,19 @@ namespace GLTF
 			TexExpression->SetGroupName(*GroupName);
 
 			CreateTextureCoordinate(CoordinateIndex, *TexExpression, *CurrentMaterialElement);
+
+			if (nullptr != TextureTransform)
+			{
+				CreateTextureTransform(*TextureTransform, *TexExpression, *CurrentMaterialElement);
+			}
 		}
 		return TexExpression;
 	}
 
 	template <class ValueExpressionClass, class ValueClass>
 	FMaterialExpression* FPBRMapFactory::CreateMap(const GLTF::FTexture& Map, int CoordinateIndex, const ValueClass& Value, const TCHAR* MapName,
-	                                               const TCHAR* ValueName, ETextureMode TextureMode, FMaterialExpressionInput& MaterialInput)
+	                                               const TCHAR* ValueName, ETextureMode TextureMode, FMaterialExpressionInput& MaterialInput, 
+												   const GLTF::FTextureTransform* TextureTransform)
 	{
 		check(MapName);
 		check(CurrentMaterialElement);
@@ -249,7 +322,7 @@ namespace GLTF
 		ValueExpression->SetGroupName(*GroupName);
 		SetExpresisonValue(Value, ValueExpression);
 
-		FMaterialExpressionTexture* TexExpression = CreateTextureMap(Map, CoordinateIndex, MapName, TextureMode);
+		FMaterialExpressionTexture* TexExpression = CreateTextureMap(Map, CoordinateIndex, MapName, TextureMode, TextureTransform);
 		if (TexExpression)
 		{
 			FMaterialExpressionGeneric* MultiplyExpression = CurrentMaterialElement->AddMaterialExpression<FMaterialExpressionGeneric>();
diff --git a/Engine/Plugins/Editor/GLTFImporter/Source/GLTFCore/Private/GLTFMapFactory.h b/Engine/Plugins/Editor/GLTFImporter/Source/GLTFCore/Private/GLTFMapFactory.h
--- a/Engine/Plugins/Editor/GLTFImporter/Source/GLTFCore/Private/GLTFMapFactory.h
+++ b/Engine/Plugins/Editor/GLTFImporter/Source/GLTFCore/Private/GLTFMapFactory.h
@@ -10,6 +10,7 @@
 namespace GLTF
 {
 	struct FTexture;
+	struct FTextureTransform;
 	class FTextureFactory;
 
 	struct FPBRMapFactory
@@ -74,25 +75,29 @@ namespace GLTF
 
 		void SetParentPackage(UObject* ParentPackage, EObjectFlags Flags);
 
-		void CreateNormalMap(const GLTF::FTexture& Map, int CoordinateIndex, float NormalScale);
+		void CreateNormalMap(const GLTF::FTexture& Map, int CoordinateIndex, float NormalScale, const GLTF::FTextureTransform* TextureTransform);
 
 		FMaterialExpression* CreateColorMap(const GLTF::FTexture& Map, int CoordinateIndex, const FVector3f& Color, const TCHAR* MapName,
-		                                    const TCHAR* ValueName, ETextureMode TextureMode, FMaterialExpressionInput& MaterialInput);
+		                                    const TCHAR* ValueName, ETextureMode TextureMode, FMaterialExpressionInput& MaterialInput,
+											const GLTF::FTextureTransform* TextureTransform);
 
 		FMaterialExpression* CreateColorMap(const GLTF::FTexture& Map, int CoordinateIndex, const FVector4f& Color, const TCHAR* MapName,
-		                                    const TCHAR* ValueName, ETextureMode TextureMode, FMaterialExpressionInput& MaterialInput);
+		                                    const TCHAR* ValueName, ETextureMode TextureMode, FMaterialExpressionInput& MaterialInput,
+											const GLTF::FTextureTransform* TextureTransform);
 
-		FMaterialExpressionTexture* CreateTextureMap(const GLTF::FTexture& Map, int CoordinateIndex, const TCHAR* MapName, ETextureMode TextureMode);
+		FMaterialExpressionTexture* CreateTextureMap(const GLTF::FTexture& Map, int CoordinateIndex, const TCHAR* MapName, ETextureMode TextureMode, 
+													 const GLTF::FTextureTransform* TextureTransform);
 
 		void CreateMultiMap(const GLTF::FTexture& Map, int CoordinateIndex, const TCHAR* MapName, const FMapChannel* MapChannels,
-		                    uint32 MapChannelsCount, ETextureMode TextureMode);
+		                    uint32 MapChannelsCount, ETextureMode TextureMode, const GLTF::FTextureTransform* TextureTransform);
 
 	private:
 		using FExpressionList = TArray<FMaterialExpression*, TFixedAllocator<4> >;
 
 		template <class ValueExpressionClass, class ValueClass>
 		FMaterialExpression* CreateMap(const GLTF::FTexture& Map, int CoordinateIndex, const ValueClass& Value, const TCHAR* MapName,
-		                               const TCHAR* ValueName, ETextureMode TextureMode, FMaterialExpressionInput& MaterialInput);
+		                               const TCHAR* ValueName, ETextureMode TextureMode, FMaterialExpressionInput& MaterialInput, 
+									   const GLTF::FTextureTransform* TextureTransform);
 
 	private:
 		ITextureFactory& TextureFactory;
diff --git a/Engine/Plugins/Editor/GLTFImporter/Source/GLTFCore/Private/GLTFMaterialFactory.cpp b/Engine/Plugins/Editor/GLTFImporter/Source/GLTFCore/Private/GLTFMaterialFactory.cpp
--- a/Engine/Plugins/Editor/GLTFImporter/Source/GLTFCore/Private/GLTFMaterialFactory.cpp
+++ b/Engine/Plugins/Editor/GLTFImporter/Source/GLTFCore/Private/GLTFMaterialFactory.cpp
@@ -15,8 +15,8 @@ namespace GLTF
 	{
 		const GLTF::FTexture& GetTexture(const GLTF::FTextureMap& Map, const TArray<GLTF::FTexture>& Textures)
 		{
-			static const GLTF::FImage   Immage;
-			static const GLTF::FTexture None(FString(), Immage, GLTF::FSampler::DefaultSampler);
+			static const GLTF::FImage   Image;
+			static const GLTF::FTexture None(FString(), Image, GLTF::FSampler::DefaultSampler);
 			return Map.TextureIndex != INDEX_NONE ? Textures[Map.TextureIndex] : None;
 		}
 	}
@@ -219,7 +219,8 @@ namespace GLTF
 									  TEXT("BaseColor"),
 									  nullptr,
 									  ETextureMode::Color,
-									  MaterialElement.GetBaseColor());
+									  MaterialElement.GetBaseColor(),
+									  GLTFMaterial.BaseColor.bHasTextureTransform ? &GLTFMaterial.BaseColor.TextureTransform : nullptr);
 
 			// Metallic
 			Maps.Emplace(GLTFMaterial.MetallicRoughness.MetallicFactor,
@@ -240,7 +241,8 @@ namespace GLTF
 				                      TEXT("MetallicRoughness"),
 									  Maps.GetData(),
 									  Maps.Num(),
-									  ETextureMode::Grayscale);
+									  ETextureMode::Grayscale,
+									  GLTFMaterial.MetallicRoughness.Map.bHasTextureTransform ? &GLTFMaterial.MetallicRoughness.Map.TextureTransform : nullptr);
 		}
 		else if (GLTFMaterial.ShadingModel == FMaterial::EShadingModel::SpecularGlossiness)
 		{
@@ -262,7 +264,8 @@ namespace GLTF
 																	 TEXT("Diffuse"),
 																	 TEXT("Color"),
 																	 ETextureMode::Color,
-																	 *SpecGlossToMetalRough->GetInput(1));
+																	 *SpecGlossToMetalRough->GetInput(1),
+																	 GLTFMaterial.BaseColor.bHasTextureTransform ? &GLTFMaterial.BaseColor.TextureTransform : nullptr);
 
 			// Specular (goes into SpecGlossToMetalRough conversion)
 			Maps.Emplace(GLTFMaterial.SpecularGlossiness.SpecularFactor,
@@ -284,7 +287,8 @@ namespace GLTF
 				                      TEXT("SpecularGlossiness"),
 									  Maps.GetData(),
 									  Maps.Num(),
-				                      ETextureMode::Color);
+				                      ETextureMode::Color,
+									  GLTFMaterial.SpecularGlossiness.Map.bHasTextureTransform ? &GLTFMaterial.SpecularGlossiness.Map.TextureTransform : nullptr);
 		}
 	}
 
@@ -293,7 +297,8 @@ namespace GLTF
 		MapFactory.GroupName = TEXT("Occlusion");
 
 		FMaterialExpressionTexture* TexExpression = MapFactory.CreateTextureMap(
-			GetTexture(GLTFMaterial.Occlusion, Textures), GLTFMaterial.Occlusion.TexCoord, TEXT("Occlusion"), ETextureMode::Grayscale);
+			GetTexture(GLTFMaterial.Occlusion, Textures), GLTFMaterial.Occlusion.TexCoord, TEXT("Occlusion"), ETextureMode::Grayscale, 
+			GLTFMaterial.Occlusion.bHasTextureTransform ? &GLTFMaterial.Occlusion.TextureTransform : nullptr);
 
 		if (!TexExpression)
 			return;
@@ -330,7 +335,8 @@ namespace GLTF
 								  TEXT("Emissive"),
 								  TEXT("Color"),
 								  ETextureMode::Color, // emissive map is in sRGB space
-								  MaterialElement.GetEmissiveColor());
+								  MaterialElement.GetEmissiveColor(),
+								  GLTFMaterial.Emissive.bHasTextureTransform ? &GLTFMaterial.Emissive.TextureTransform : nullptr);
 	}
 
 	void FMaterialFactoryImpl::HandleNormal(const TArray<GLTF::FTexture>& Textures, const GLTF::FMaterial& GLTFMaterial, FPBRMapFactory& MapFactory, FMaterialElement& MaterialElement)
@@ -338,7 +344,8 @@ namespace GLTF
 		MapFactory.GroupName = TEXT("Normal");
 		MapFactory.CreateNormalMap(GetTexture(GLTFMaterial.Normal, Textures),
 								   GLTFMaterial.Normal.TexCoord,
-								   GLTFMaterial.NormalScale);
+								   GLTFMaterial.NormalScale,
+								   GLTFMaterial.Normal.bHasTextureTransform ? &GLTFMaterial.Normal.TextureTransform : nullptr);
 	}
 
 	void FMaterialFactoryImpl::HandleClearCoat(const TArray<GLTF::FTexture>& Textures, const GLTF::FMaterial& GLTFMaterial, FPBRMapFactory& MapFactory, FMaterialElement& MaterialElement)
@@ -362,13 +369,15 @@ namespace GLTF
 			GetTexture(GLTFMaterial.ClearCoat.ClearCoatMap, Textures),
 			GLTFMaterial.ClearCoat.ClearCoatMap.TexCoord,
 			TEXT("ClearCoat"),
-			ETextureMode::Color);
+			ETextureMode::Color,
+			GLTFMaterial.ClearCoat.ClearCoatMap.bHasTextureTransform ? &GLTFMaterial.ClearCoat.ClearCoatMap.TextureTransform : nullptr);
 
 		FMaterialExpressionTexture* ClearCoatRoughnessTexture = MapFactory.CreateTextureMap(
 			GetTexture(GLTFMaterial.ClearCoat.RoughnessMap, Textures),
 			GLTFMaterial.ClearCoat.RoughnessMap.TexCoord,
 			TEXT("ClearCoatRoughness"),
-			ETextureMode::Color);
+			ETextureMode::Color,
+			GLTFMaterial.ClearCoat.RoughnessMap.bHasTextureTransform ? &GLTFMaterial.ClearCoat.RoughnessMap.TextureTransform : nullptr);
 
 		FMaterialExpression* ClearCoatExpr = ClearCoatFactor;
 		if (ClearCoatTexture)
@@ -399,7 +408,8 @@ namespace GLTF
 			GetTexture(GLTFMaterial.ClearCoat.NormalMap, Textures),
 			GLTFMaterial.ClearCoat.NormalMap.TexCoord,
 			TEXT("ClearCoatNormal"),
-			ETextureMode::Normal);
+			ETextureMode::Normal,
+			GLTFMaterial.ClearCoat.NormalMap.bHasTextureTransform ? &GLTFMaterial.ClearCoat.NormalMap.TextureTransform : nullptr);
 
 		if (ClearCoatNormalTexture)
 		{
@@ -446,7 +456,8 @@ namespace GLTF
 			GetTexture(GLTFMaterial.Transmission.TransmissionMap, Textures),
 			GLTFMaterial.Transmission.TransmissionMap.TexCoord,
 			TEXT("Transmission"),
-			ETextureMode::Color);
+			ETextureMode::Color,
+			GLTFMaterial.Transmission.TransmissionMap.bHasTextureTransform ? &GLTFMaterial.Transmission.TransmissionMap.TextureTransform : nullptr);
 
 		FMaterialExpression* MetallicFactor = MaterialElement.GetMetallic().GetExpression();
 		FMaterialExpression* RoughnessFactor = MaterialElement.GetRoughness().GetExpression();
@@ -719,7 +730,8 @@ namespace GLTF
 		SpecularFactorExpr->GetScalar() = GLTFMaterial.Specular.SpecularFactor;
 
 		FMaterialExpressionTexture* SpecularTextureExpr = MapFactory.CreateTextureMap(
-			GetTexture(GLTFMaterial.Specular.SpecularMap, Textures), GLTFMaterial.Specular.SpecularMap.TexCoord, TEXT("SpecularMap"), ETextureMode::Color);
+			GetTexture(GLTFMaterial.Specular.SpecularMap, Textures), GLTFMaterial.Specular.SpecularMap.TexCoord, TEXT("SpecularMap"), ETextureMode::Color,
+			GLTFMaterial.Specular.SpecularMap.bHasTextureTransform ? &GLTFMaterial.Specular.SpecularMap.TextureTransform : nullptr);
 
 		FMaterialExpression* SpecularExpr = SpecularFactorExpr;
 
diff --git a/Engine/Plugins/Editor/GLTFImporter/Source/GLTFCore/Public/GLTFMaterial.h b/Engine/Plugins/Editor/GLTFImporter/Source/GLTFCore/Public/GLTFMaterial.h
--- a/Engine/Plugins/Editor/GLTFImporter/Source/GLTFCore/Public/GLTFMaterial.h
+++ b/Engine/Plugins/Editor/GLTFImporter/Source/GLTFCore/Public/GLTFMaterial.h
@@ -7,14 +7,32 @@
 
 namespace GLTF
 {
+	struct GLTFCORE_API FTextureTransform
+	{
+		float Offset[2];
+		float Rotation;
+		float Scale[2];
+
+		FTextureTransform()
+		{
+			Offset[0] = Offset[1] = 0.0f;
+			Scale[0] = Scale[1] = 1.0f;
+			Rotation = 0.0f;
+		}
+	};
+
 	struct GLTFCORE_API FTextureMap
 	{
 		int32 TextureIndex;
 		uint8 TexCoord;
 
+		bool bHasTextureTransform;
+		FTextureTransform TextureTransform;
+
 		FTextureMap()
 		    : TextureIndex(INDEX_NONE)
 		    , TexCoord(0)
+			, bHasTextureTransform(false)
 		{
 		}
 
