GLTF: Fix issues caused by material name collisions.

diff --git a/Engine/Plugins/Editor/GLTFImporter/Source/GLTFCore/Private/GLTF/GLTFAsset.cpp b/Engine/Plugins/Editor/GLTFImporter/Source/GLTFCore/Private/GLTF/GLTFAsset.cpp
--- a/Engine/Plugins/Editor/GLTFImporter/Source/GLTFCore/Private/GLTF/GLTFAsset.cpp
+++ b/Engine/Plugins/Editor/GLTFImporter/Source/GLTFCore/Private/GLTF/GLTFAsset.cpp
@@ -11,13 +11,16 @@ namespace GLTF
 		template <typename T>
 		void GenerateNames(const FString& Prefix, TArray<T>& Objects)
 		{
-			int32 Counter = 0;
-			for (T& Obj : Objects)
+			for (int32 Index = 0; Index < Objects.Num(); ++Index)
 			{
+				T& Obj = Objects[Index];
+
 				if (Obj.Name.IsEmpty())
 				{
-					Obj.Name = Prefix + FString::FromInt(Counter++);
+					Obj.Name = Prefix;
 				}
+
+				Obj.Name = FString::FromInt(Index) + TEXT("_") + Obj.Name;
 			}
 		}
 	}
@@ -59,15 +62,9 @@ namespace GLTF
 	{
 		check(!Prefix.IsEmpty());
 
-		GLTF::GenerateNames(Prefix + TEXT("_material_"), Materials);
-		GLTF::GenerateNames(Prefix + TEXT("_skin_"), Skins);
-		GLTF::GenerateNames(Prefix + TEXT("_animation_"), Animations);
-
 		{
-			const FString& NodePrefix = Prefix + TEXT("_node_");
-			const FString& JoinPrefix = Prefix + TEXT("_join_");
-
-			int32 Counter[2] = {0, 0};
+			const FString NodePrefix = Prefix + TEXT("_node_");
+			const FString JointPrefix = Prefix + TEXT("_joint_");
 			for (int32 NodeIndex = 0; NodeIndex < Nodes.Num(); ++NodeIndex)
 			{
 				FNode& Node = Nodes[NodeIndex];
@@ -75,7 +72,7 @@ namespace GLTF
 				if (Node.Name.IsEmpty())
 				{
 					bool bIsJoint = Node.Type == FNode::EType::Joint;
-					Node.Name     = (bIsJoint ? JoinPrefix : NodePrefix) + FString::FromInt(Counter[bIsJoint]++);
+					Node.Name     = (bIsJoint ? JointPrefix : NodePrefix);
 				}
 
 				// Make sure node names are unique
@@ -84,8 +81,8 @@ namespace GLTF
 		}
 
 		{
-			const FString& TexPrefix = Prefix + TEXT("_texture_");
-			for (int32 TextureIndex = 0; TextureIndex != Textures.Num(); ++TextureIndex)
+			const FString TexPrefix = Prefix + TEXT("_texture_");
+			for (int32 TextureIndex = 0; TextureIndex < Textures.Num(); ++TextureIndex)
 			{
 				FTexture& Tex = Textures[TextureIndex];
 
@@ -114,44 +111,60 @@ namespace GLTF
 		}
 
 		{
-			const FString& MeshPrefix = Prefix + TEXT("_mesh_");
-			for (int32 Index = 0; Index != Meshes.Num(); ++Index)
+			const FString CameraPrefix = Prefix + TEXT("_camera_");
+			for (int32 CameraIndex = 0; CameraIndex < Cameras.Num(); ++CameraIndex)
 			{
-				FMesh& Mesh = Meshes[Index];
+				FCamera& Camera = Cameras[CameraIndex];
 
-				if (Mesh.Name.IsEmpty())
+				if (Camera.Name.IsEmpty())
 				{
-					Mesh.Name = MeshPrefix;
+					if (!Camera.Node.Name.IsEmpty())
+					{
+						Camera.Name = Camera.Node.Name; // cant be empty
+					}
+					else
+					{
+						Camera.Name = CameraPrefix;
+					}
 				}
 
-				// GLTF mesh name has decorative purpose, not guaranteed to be unique
-				// only its index is unique. Same with glTF image or its source file's basename
-				// So always include mesh index into meshes name to increase probability that names are unique
-				Mesh.Name = FString::FromInt(Index) + TEXT("_") + Mesh.Name;
+				Camera.Name = FString::FromInt(CameraIndex) + TEXT("_") + Camera.Name;
 			}
 		}
 
-		for (FCamera& Camera : Cameras)
 		{
-			if (Camera.Name.IsEmpty())
-				Camera.Name = TEXT("camera_") + Camera.Node.Name;  // cant be empty
-		}
-
-		int32 Counter = 0;
-		for (FLight& Light : Lights)
-		{
-			if (Light.Name.IsEmpty())
+			const FString LightPrefix = Prefix + TEXT("_light_");
+			for (int32 LightIndex = 0; LightIndex < Lights.Num(); ++LightIndex)
 			{
-				// cant be empty
-				if (Light.Node)
-					Light.Name = TEXT("light_") + Light.Node->Name;
-				else
-					Light.Name = TEXT("light_") + FString::FromInt(Counter++);
+				FLight& Light = Lights[LightIndex];
+
+				if (Light.Name.IsEmpty())
+				{
+					if (Light.Node && !Light.Node->Name.IsEmpty())
+					{
+						Light.Name = Light.Node->Name;
+					}
+					else
+					{
+						Light.Name = LightPrefix;
+					}
+				}
+
+				Light.Name = FString::FromInt(LightIndex) + TEXT("_") + Light.Name;
 			}
 		}
 
-		GLTF::GenerateNames(Prefix + TEXT("_image_"), Images);
-		GLTF::GenerateNames(Prefix + TEXT("_image_"), Meshes);
+		// General pattern here is that we'll have a Prefix that will be a filename (like MyFile). If the
+		// GLTF object has an actual name, it will end up like <index>_<objectname>, like "0_shoe" or "3_spotlight".
+		// If it doesn't have a name, it will end up with a name like <index>_<filename>_<objecttype>, e.g.
+		// "0_MyFile_material" or "3_myfile_mesh".
+		// We want to do images after textures as the textures may want to reuse source image names, but we don't
+		// want to use those if we generated them ourselves.
+		GLTF::GenerateNames(Prefix + TEXT("_material"), Materials);
+		GLTF::GenerateNames(Prefix + TEXT("_skin"), Skins);
+		GLTF::GenerateNames(Prefix + TEXT("_animation"), Animations);
+		GLTF::GenerateNames(Prefix + TEXT("_image"), Images);
+		GLTF::GenerateNames(Prefix + TEXT("_mesh"), Meshes);
 	}
 
 	void FAsset::GetRootNodes(TArray<int32>& NodeIndices)
