Scale of gltf imported texture are wrong

diff --git a/Engine/Plugins/Editor/GLTFImporter/Source/GLTFCore/Private/GLTFMapFactory.cpp b/Engine/Plugins/Editor/GLTFImporter/Source/GLTFCore/Private/GLTFMapFactory.cpp
--- a/Engine/Plugins/Editor/GLTFImporter/Source/GLTFCore/Private/GLTFMapFactory.cpp
+++ b/Engine/Plugins/Editor/GLTFImporter/Source/GLTFCore/Private/GLTFMapFactory.cpp
@@ -16,7 +16,7 @@ namespace GLTF
 			if (TexCoord)
 			{
 				FMaterialExpressionTextureCoordinate* CoordExpression = MaterialElement.AddMaterialExpression<FMaterialExpressionTextureCoordinate>();
-				CoordExpression->SetCoordinateIndex(TexCoord + 1);
+				CoordExpression->SetCoordinateIndex(TexCoord);
 				CoordExpression->ConnectExpression(TexExpression.GetInputCoordinate(), 0);
 			}
 		}
@@ -58,9 +58,15 @@ namespace GLTF
 				}
 
 				FMaterialExpressionScalar* RotationAngle = MaterialElement.AddMaterialExpression<FMaterialExpressionScalar>();
-				RotationAngle->GetScalar() = AngleRadians / TWO_PI; // Normalize angle value
+				RotationAngle->GetScalar() = 1.0f - (AngleRadians / TWO_PI); // Normalize angle value
+
+				FMaterialExpressionGeneric* RotationPivot = MaterialElement.AddMaterialExpression<FMaterialExpressionGeneric>();
+				RotationPivot->SetExpressionName(TEXT("Constant2Vector"));
+				RotationPivot->SetFloatProperty(TEXT("R"), 0.0f);
+				RotationPivot->SetFloatProperty(TEXT("G"), 0.0f);
 
 				LastExprInChain->ConnectExpression(*UVRotate->GetInput(0), 0);
+				RotationPivot->ConnectExpression(*UVRotate->GetInput(1), 0);
 				RotationAngle->ConnectExpression(*UVRotate->GetInput(2), 0);
 
 				LastExprInChain = UVRotate;
