Datasmith glTF: Editor crash when imported some gltf with embedded texture

diff --git a/Engine/Plugins/Editor/GLTFImporter/Source/GLTFCore/Private/GLTF/GLTFReader.cpp b/Engine/Plugins/Editor/GLTFImporter/Source/GLTFCore/Private/GLTF/GLTFReader.cpp
--- a/Engine/Plugins/Editor/GLTFImporter/Source/GLTFCore/Private/GLTF/GLTFReader.cpp
+++ b/Engine/Plugins/Editor/GLTFImporter/Source/GLTFCore/Private/GLTF/GLTFReader.cpp
@@ -125,7 +125,7 @@ namespace GLTF
 				FString MimeType;
 				uint32  DataSize = 0;
 				bool    bSuccess = DecodeDataURI(URI, MimeType, CurrentBufferOffset, DataSize);
-				if (!bSuccess || MimeType != TEXT("application/octet-stream") || !ensure(DataSize == ByteLength))
+				if (!bSuccess || (MimeType != TEXT("application/octet-stream") && MimeType != TEXT("application/gltf-buffer")) || !ensure(DataSize == ByteLength))
 				{
 					Messages.Emplace(EMessageSeverity::Error, TEXT("Problem decoding buffer from data URI."));
 				}
@@ -735,7 +735,7 @@ namespace GLTF
 				{
 					FString      MimeType;
 					const uint32 DataSize = GetDecodedDataSize(URI, MimeType);
-					if (DataSize > 0 && MimeType == TEXT("application/octet-stream"))
+					if (DataSize > 0 && (MimeType == TEXT("application/octet-stream") || MimeType == (TEXT("application/gltf-buffer"))))
 					{
 						ensure(DataSize == ByteLength);
 						ExtraBufferSize += ByteLength;
