When importing gltf with MaterialsSpecular extension or MaterialsIOR extension I get warning of unkwnonExtension

diff --git a/Engine/Plugins/Editor/GLTFImporter/Source/GLTFCore/Private/GLTF/GLTFAsset.cpp b/Engine/Plugins/Editor/GLTFImporter/Source/GLTFCore/Private/GLTF/GLTFAsset.cpp
--- a/Engine/Plugins/Editor/GLTFImporter/Source/GLTFCore/Private/GLTF/GLTFAsset.cpp
+++ b/Engine/Plugins/Editor/GLTFImporter/Source/GLTFCore/Private/GLTF/GLTFAsset.cpp
@@ -30,6 +30,16 @@ namespace GLTF
 		return Extras.FindByPredicate([&Name](const FMetadata::FExtraData& Data) { return Data.Name == Name; });
 	}
 
+	const TSet<GLTF::EExtension> FAsset::SupportedExtensions = { GLTF::EExtension::KHR_MaterialsPbrSpecularGlossiness,
+																 GLTF::EExtension::KHR_MaterialsUnlit,
+																 GLTF::EExtension::KHR_LightsPunctual,
+																 GLTF::EExtension::KHR_MaterialsClearCoat,
+																 GLTF::EExtension::KHR_MaterialsTransmission,
+																 GLTF::EExtension::KHR_MaterialsSheen,
+																 GLTF::EExtension::KHR_MaterialsVariants,
+																 GLTF::EExtension::KHR_MaterialsSpecular,
+																 GLTF::EExtension::KHR_MaterialsIOR };
+
 	void FAsset::Clear(uint32 BinBufferKBytes, uint32 ExtraBinBufferKBytes)
 	{
 		Buffers.Empty();
@@ -227,6 +237,10 @@ namespace GLTF
 				return TEXT("KHR_MaterialsSheen");
 			case GLTF::EExtension::KHR_MaterialsTransmission:
 				return TEXT("KHR_MaterialsTransmission");
+			case GLTF::EExtension::KHR_MaterialsSpecular:
+				return TEXT("KHR_MaterialsSpecular");
+			case GLTF::EExtension::KHR_MaterialsIOR:
+				return TEXT("KHR_MaterialsIOR");
 			case GLTF::EExtension::KHR_TextureTransform:
 				return TEXT("KHR_Texture_Transform");
 			case GLTF::EExtension::KHR_DracoMeshCompression:
diff --git a/Engine/Plugins/Editor/GLTFImporter/Source/GLTFCore/Public/GLTFAsset.h b/Engine/Plugins/Editor/GLTFImporter/Source/GLTFCore/Public/GLTFAsset.h
--- a/Engine/Plugins/Editor/GLTFImporter/Source/GLTFCore/Public/GLTFAsset.h
+++ b/Engine/Plugins/Editor/GLTFImporter/Source/GLTFCore/Public/GLTFAsset.h
@@ -55,6 +55,8 @@ namespace GLTF
 
 	struct GLTFCORE_API FAsset : FNoncopyable
 	{
+		static const TSet<EExtension> SupportedExtensions;
+
 		enum EValidationCheck
 		{
 			Valid,
diff --git a/Engine/Plugins/Editor/GLTFImporter/Source/GLTFImporter/Private/GLTFImporterContext.cpp b/Engine/Plugins/Editor/GLTFImporter/Source/GLTFImporter/Private/GLTFImporterContext.cpp
--- a/Engine/Plugins/Editor/GLTFImporter/Source/GLTFImporter/Private/GLTFImporterContext.cpp
+++ b/Engine/Plugins/Editor/GLTFImporter/Source/GLTFImporter/Private/GLTFImporterContext.cpp
@@ -42,16 +42,9 @@ bool FGLTFImporterContext::OpenFile(const FString& FilePath)
 	check(Asset.ValidationCheck() == GLTF::FAsset::Valid);
 
 	// check extensions supported
-	static const TArray<GLTF::EExtension> SupportedExtensions = {GLTF::EExtension::KHR_MaterialsPbrSpecularGlossiness,
-	                                                             GLTF::EExtension::KHR_MaterialsUnlit,
-																 GLTF::EExtension::KHR_MaterialsClearCoat,
-																 GLTF::EExtension::KHR_MaterialsTransmission,
-																 GLTF::EExtension::KHR_MaterialsSheen,
-																 GLTF::EExtension::KHR_MaterialsVariants};
-
 	for (GLTF::EExtension Extension : Asset.ExtensionsUsed)
 	{
-		if (SupportedExtensions.Find(Extension) == INDEX_NONE)
+		if (!GLTF::FAsset::SupportedExtensions.Contains(Extension))
 		{
 			LogMessages.Emplace(GLTF::EMessageSeverity::Warning, FString::Printf(TEXT("Extension is not supported: %s"), GLTF::ToString(Extension)));
 		}
diff --git a/Engine/Plugins/Enterprise/DatasmithGLTFImporter/Source/DatasmithGLTFTranslator/Private/DatasmithGLTFImporter.cpp b/Engine/Plugins/Enterprise/DatasmithGLTFImporter/Source/DatasmithGLTFTranslator/Private/DatasmithGLTFImporter.cpp
--- a/Engine/Plugins/Enterprise/DatasmithGLTFImporter/Source/DatasmithGLTFTranslator/Private/DatasmithGLTFImporter.cpp
+++ b/Engine/Plugins/Enterprise/DatasmithGLTFImporter/Source/DatasmithGLTFTranslator/Private/DatasmithGLTFImporter.cpp
@@ -99,17 +99,9 @@ bool FDatasmithGLTFImporter::OpenFile(const FString& InFileName)
 	GLTFAsset->GenerateNames(FPaths::GetBaseFilename(InFileName));
 
 	// check extensions supported
-	static const TArray<GLTF::EExtension> SupportedExtensions = {GLTF::EExtension::KHR_MaterialsPbrSpecularGlossiness,
-																 GLTF::EExtension::KHR_MaterialsUnlit,
-																 GLTF::EExtension::KHR_LightsPunctual,
-																 GLTF::EExtension::KHR_MaterialsClearCoat,
-																 GLTF::EExtension::KHR_MaterialsTransmission,
-																 GLTF::EExtension::KHR_MaterialsSheen,
-																 GLTF::EExtension::KHR_MaterialsVariants};
-
 	for (GLTF::EExtension Extension : GLTFAsset->ExtensionsUsed)
 	{
-		if (SupportedExtensions.Find(Extension) == INDEX_NONE)
+		if (!GLTF::FAsset::SupportedExtensions.Contains(Extension))
 		{
 			LogMessages.Emplace(GLTF::EMessageSeverity::Warning, FString::Printf(TEXT("Extension is not supported: %s"), GLTF::ToString(Extension)));
 		}
