glTF import of unlit material generates a default lit material

diff --git a/Engine/Plugins/Editor/GLTFImporter/Source/GLTFCore/Private/GLTFMaterialFactory.cpp b/Engine/Plugins/Editor/GLTFImporter/Source/GLTFCore/Private/GLTFMaterialFactory.cpp
--- a/Engine/Plugins/Editor/GLTFImporter/Source/GLTFCore/Private/GLTFMaterialFactory.cpp
+++ b/Engine/Plugins/Editor/GLTFImporter/Source/GLTFCore/Private/GLTFMaterialFactory.cpp
@@ -208,6 +208,14 @@ namespace GLTF
 	{
 		MapFactory.GroupName = TEXT("GGX");
 
+		FMaterialExpressionInput* BaseColorInput = &MaterialElement.GetBaseColor();
+
+		if (GLTFMaterial.bIsUnlitShadingModel)
+		{
+			BaseColorInput = &MaterialElement.GetEmissiveColor();
+			MaterialElement.SetShadingModel(EGLTFMaterialShadingModel::Unlit);
+		}
+
 		TArray<FPBRMapFactory::FMapChannel, TFixedAllocator<4>> Maps;
 		if (GLTFMaterial.ShadingModel == FMaterial::EShadingModel::MetallicRoughness)
 		{
@@ -219,76 +227,83 @@ namespace GLTF
 									  TEXT("BaseColor"),
 									  nullptr,
 									  ETextureMode::Color,
-									  MaterialElement.GetBaseColor(),
+									  *BaseColorInput,
 									  GLTFMaterial.BaseColor.bHasTextureTransform ? &GLTFMaterial.BaseColor.TextureTransform : nullptr);
 
-			// Metallic
-			Maps.Emplace(GLTFMaterial.MetallicRoughness.MetallicFactor,
-						 TEXT("Metallic Factor"),
-						 FPBRMapFactory::EChannel::Blue,
-						 &MaterialElement.GetMetallic(),
-						 nullptr);
-
-			// Roughness
-			Maps.Emplace(GLTFMaterial.MetallicRoughness.RoughnessFactor,
-						 TEXT("Roughness Factor"),
-						 FPBRMapFactory::EChannel::Green,
-						 &MaterialElement.GetRoughness(),
-						 nullptr);
-
-			MapFactory.CreateMultiMap(GetTexture(GLTFMaterial.MetallicRoughness.Map, Textures),
-									  GLTFMaterial.MetallicRoughness.Map.TexCoord,
-				                      TEXT("MetallicRoughness"),
-									  Maps.GetData(),
-									  Maps.Num(),
-									  ETextureMode::Grayscale,
-									  GLTFMaterial.MetallicRoughness.Map.bHasTextureTransform ? &GLTFMaterial.MetallicRoughness.Map.TextureTransform : nullptr);
+			if (!GLTFMaterial.bIsUnlitShadingModel)
+			{
+				// Metallic
+				Maps.Emplace(GLTFMaterial.MetallicRoughness.MetallicFactor,
+							 TEXT("Metallic Factor"),
+							 FPBRMapFactory::EChannel::Blue,
+							 &MaterialElement.GetMetallic(),
+							 nullptr);
+
+				// Roughness
+				Maps.Emplace(GLTFMaterial.MetallicRoughness.RoughnessFactor,
+							 TEXT("Roughness Factor"),
+							 FPBRMapFactory::EChannel::Green,
+							 &MaterialElement.GetRoughness(),
+							 nullptr);
+
+				MapFactory.CreateMultiMap(GetTexture(GLTFMaterial.MetallicRoughness.Map, Textures),
+										  GLTFMaterial.MetallicRoughness.Map.TexCoord,
+										  TEXT("MetallicRoughness"),
+										  Maps.GetData(),
+										  Maps.Num(),
+										  ETextureMode::Grayscale,
+										  GLTFMaterial.MetallicRoughness.Map.bHasTextureTransform ? &GLTFMaterial.MetallicRoughness.Map.TextureTransform : nullptr);
+			}
 		}
 		else if (GLTFMaterial.ShadingModel == FMaterial::EShadingModel::SpecularGlossiness)
 		{
 			// We'll actually just convert it into MetalRoughness in the material graph
 			FMaterialExpressionFunctionCall* SpecGlossToMetalRough = MaterialElement.AddMaterialExpression<FMaterialExpressionFunctionCall>();
 			SpecGlossToMetalRough->SetFunctionPathName(TEXT("/GLTFImporter/SpecGlossToMetalRoughness.SpecGlossToMetalRoughness"));
-			SpecGlossToMetalRough->ConnectExpression(MaterialElement.GetBaseColor(), 0);
-			SpecGlossToMetalRough->ConnectExpression(MaterialElement.GetMetallic(), 1);
-
-			FMaterialExpressionGeneric* GlossToRoughness = MaterialElement.AddMaterialExpression<FMaterialExpressionGeneric>();
-			GlossToRoughness->SetExpressionName(TEXT("OneMinus"));
-			GlossToRoughness->ConnectExpression(MaterialElement.GetRoughness(), 0);
-
-			// Diffuse Color (BaseColor/BaseColorFactor are used to store the Diffuse alternatives for Spec/Gloss)
-			MapFactory.GroupName = TEXT("Diffuse Color");
-			FMaterialExpression* Diffuse = MapFactory.CreateColorMap(GetTexture(GLTFMaterial.BaseColor, Textures),
-																	 GLTFMaterial.BaseColor.TexCoord,
-																	 GLTFMaterial.BaseColorFactor,
-																	 TEXT("Diffuse"),
-																	 TEXT("Color"),
-																	 ETextureMode::Color,
-																	 *SpecGlossToMetalRough->GetInput(1),
-																	 GLTFMaterial.BaseColor.bHasTextureTransform ? &GLTFMaterial.BaseColor.TextureTransform : nullptr);
-
-			// Specular (goes into SpecGlossToMetalRough conversion)
-			Maps.Emplace(GLTFMaterial.SpecularGlossiness.SpecularFactor,
-						 TEXT("Specular Factor"),
-						 FPBRMapFactory::EChannel::RGB,
-						 SpecGlossToMetalRough->GetInput(0),
-						 nullptr);
-
-			// Glossiness (converted to Roughness)
-			Maps.Emplace(GLTFMaterial.SpecularGlossiness.GlossinessFactor,
-						 TEXT("Glossiness Factor"),
-				         FPBRMapFactory::EChannel::Alpha,
-						 GlossToRoughness->GetInput(0),
-						 nullptr);
-
-			// Creates the multimap for Specular and Glossiness
-			MapFactory.CreateMultiMap(GetTexture(GLTFMaterial.SpecularGlossiness.Map, Textures),
-									  GLTFMaterial.SpecularGlossiness.Map.TexCoord,
-				                      TEXT("SpecularGlossiness"),
-									  Maps.GetData(),
-									  Maps.Num(),
-				                      ETextureMode::Color,
-									  GLTFMaterial.SpecularGlossiness.Map.bHasTextureTransform ? &GLTFMaterial.SpecularGlossiness.Map.TextureTransform : nullptr);
+			SpecGlossToMetalRough->ConnectExpression(*BaseColorInput, 0);
+
+			if (!GLTFMaterial.bIsUnlitShadingModel)
+			{
+				SpecGlossToMetalRough->ConnectExpression(MaterialElement.GetMetallic(), 1);
+
+				FMaterialExpressionGeneric* GlossToRoughness = MaterialElement.AddMaterialExpression<FMaterialExpressionGeneric>();
+				GlossToRoughness->SetExpressionName(TEXT("OneMinus"));
+				GlossToRoughness->ConnectExpression(MaterialElement.GetRoughness(), 0);
+
+				// Diffuse Color (BaseColor/BaseColorFactor are used to store the Diffuse alternatives for Spec/Gloss)
+				MapFactory.GroupName = TEXT("Diffuse Color");
+				FMaterialExpression* Diffuse = MapFactory.CreateColorMap(GetTexture(GLTFMaterial.BaseColor, Textures),
+																		 GLTFMaterial.BaseColor.TexCoord,
+																		 GLTFMaterial.BaseColorFactor,
+																		 TEXT("Diffuse"),
+																		 TEXT("Color"),
+																		 ETextureMode::Color,
+																		 *SpecGlossToMetalRough->GetInput(1),
+																		 GLTFMaterial.BaseColor.bHasTextureTransform ? &GLTFMaterial.BaseColor.TextureTransform : nullptr);
+
+				// Specular (goes into SpecGlossToMetalRough conversion)
+				Maps.Emplace(GLTFMaterial.SpecularGlossiness.SpecularFactor,
+							 TEXT("Specular Factor"),
+							 FPBRMapFactory::EChannel::RGB,
+							 SpecGlossToMetalRough->GetInput(0),
+							 nullptr);
+
+				// Glossiness (converted to Roughness)
+				Maps.Emplace(GLTFMaterial.SpecularGlossiness.GlossinessFactor,
+							 TEXT("Glossiness Factor"),
+							 FPBRMapFactory::EChannel::Alpha,
+							 GlossToRoughness->GetInput(0),
+							 nullptr);
+
+				// Creates the multimap for Specular and Glossiness
+				MapFactory.CreateMultiMap(GetTexture(GLTFMaterial.SpecularGlossiness.Map, Textures),
+										  GLTFMaterial.SpecularGlossiness.Map.TexCoord,
+										  TEXT("SpecularGlossiness"),
+										  Maps.GetData(),
+										  Maps.Num(),
+										  ETextureMode::Color,
+										  GLTFMaterial.SpecularGlossiness.Map.bHasTextureTransform ? &GLTFMaterial.SpecularGlossiness.Map.TextureTransform : nullptr);
+			}
 		}
 	}
 
@@ -323,7 +338,7 @@ namespace GLTF
 
 	void FMaterialFactoryImpl::HandleEmissive(const TArray<GLTF::FTexture>& Textures, const GLTF::FMaterial& GLTFMaterial, FPBRMapFactory& MapFactory, FMaterialElement& MaterialElement)
 	{
-		if (GLTFMaterial.Emissive.TextureIndex == INDEX_NONE || GLTFMaterial.EmissiveFactor.IsNearlyZero())
+		if (GLTFMaterial.Emissive.TextureIndex == INDEX_NONE || GLTFMaterial.EmissiveFactor.IsNearlyZero() || GLTFMaterial.bIsUnlitShadingModel)
 		{
 			return;
 		}
@@ -350,7 +365,7 @@ namespace GLTF
 
 	void FMaterialFactoryImpl::HandleClearCoat(const TArray<GLTF::FTexture>& Textures, const GLTF::FMaterial& GLTFMaterial, FPBRMapFactory& MapFactory, FMaterialElement& MaterialElement)
 	{
-		if (!GLTFMaterial.bHasClearCoat || FMath::IsNearlyEqual(GLTFMaterial.ClearCoat.ClearCoatFactor, 0.0f))
+		if (!GLTFMaterial.bHasClearCoat || GLTFMaterial.bIsUnlitShadingModel || FMath::IsNearlyEqual(GLTFMaterial.ClearCoat.ClearCoatFactor, 0.0f))
 		{
 			return;
 		}
@@ -647,7 +662,7 @@ namespace GLTF
 	
 	void FMaterialFactoryImpl::HandleSheen(const TArray<GLTF::FTexture>& Textures, const GLTF::FMaterial& GLTFMaterial, FPBRMapFactory& MapFactory, FMaterialElement& MaterialElement)
 	{
-		if (!GLTFMaterial.bHasSheen)
+		if (!GLTFMaterial.bHasSheen || GLTFMaterial.bIsUnlitShadingModel)
 		{
 			return;
 		}
diff --git a/Engine/Plugins/Editor/GLTFImporter/Source/GLTFCore/Public/GLTFMaterialExpressions.h b/Engine/Plugins/Editor/GLTFImporter/Source/GLTFCore/Public/GLTFMaterialExpressions.h
--- a/Engine/Plugins/Editor/GLTFImporter/Source/GLTFCore/Public/GLTFMaterialExpressions.h
+++ b/Engine/Plugins/Editor/GLTFImporter/Source/GLTFCore/Public/GLTFMaterialExpressions.h
@@ -27,7 +27,8 @@ namespace GLTF
 		DefaultLit,
 		Subsurface,
 		ClearCoat,
-		ThinTranslucent
+		ThinTranslucent,
+		Unlit
 	};
 
 	class GLTFCORE_API ITextureElement
diff --git a/Engine/Plugins/Enterprise/DatasmithGLTFImporter/Source/DatasmithGLTFTranslator/Private/DatasmithGLTFMaterialElement.cpp b/Engine/Plugins/Enterprise/DatasmithGLTFImporter/Source/DatasmithGLTFTranslator/Private/DatasmithGLTFMaterialElement.cpp
--- a/Engine/Plugins/Enterprise/DatasmithGLTFImporter/Source/DatasmithGLTFTranslator/Private/DatasmithGLTFMaterialElement.cpp
+++ b/Engine/Plugins/Enterprise/DatasmithGLTFImporter/Source/DatasmithGLTFTranslator/Private/DatasmithGLTFMaterialElement.cpp
@@ -48,6 +48,9 @@ void FDatasmithGLTFMaterialElement::SetShadingModel(GLTF::EGLTFMaterialShadingMo
 
 	switch (InShadingModel)
 	{
+		case GLTF::EGLTFMaterialShadingModel::Unlit: 
+			DatasmithShadingModel = EDatasmithShadingModel::Unlit; 
+			break;
 		case GLTF::EGLTFMaterialShadingModel::ClearCoat: 
 			DatasmithShadingModel = EDatasmithShadingModel::ClearCoat; 
 			break;
