Added dllentries for UGLTFImportOptions
#9194

diff --git a/Engine/Plugins/Editor/GLTFImporter/Source/GLTFImporter/Private/GLTFImportOptions.h b/Engine/Plugins/Editor/GLTFImporter/Source/GLTFImporter/Private/GLTFImportOptions.h
--- a/Engine/Plugins/Editor/GLTFImporter/Source/GLTFImporter/Private/GLTFImportOptions.h
+++ b/Engine/Plugins/Editor/GLTFImporter/Source/GLTFImporter/Private/GLTFImportOptions.h
@@ -10,7 +10,7 @@
 #include "GLTFImportOptions.generated.h"
 
 UCLASS(config = EditorPerProjectUserSettings, HideCategories = (DebugProperty))
-class UGLTFImportOptions : public UObject
+class GLTFIMPORTER_API UGLTFImportOptions : public UObject
 {
 	GENERATED_UCLASS_BODY()
 
