[glTF] Fixed importing gltf exported from 3ds Max 2023

- ByteStride attribute for vertex data taken into account. Max 2023 glTF exporter makes use of it

- proper fix(a temporary one, not involving header changes, is being done for 5.0.3)

diff --git a/Engine/Plugins/Editor/GLTFImporter/Source/GLTFCore/Private/GLTF/GLTFAccessor.cpp b/Engine/Plugins/Editor/GLTFImporter/Source/GLTFCore/Private/GLTF/GLTFAccessor.cpp
--- a/Engine/Plugins/Editor/GLTFImporter/Source/GLTFCore/Private/GLTF/GLTFAccessor.cpp
+++ b/Engine/Plugins/Editor/GLTFImporter/Source/GLTFCore/Private/GLTF/GLTFAccessor.cpp
@@ -116,6 +116,45 @@ namespace GLTF
 			else
 				check(false);
 		}
+
+		// Copy data items that don't need conversion/expansion(i.e. Vec3 to Vec3, uint8 to uint8(not uint16)
+		// but including normalized from fixed-point types
+		template<typename ItemType, uint32 ItemElementCount>
+		void CopyWithoutConversion(const FValidAccessor& Accessor, ItemType* Buffer)
+		{
+			// Stride equals item size => use simpler copy
+			if ((Accessor.ByteStride == 0) || Accessor.ByteStride == sizeof(ItemType))
+			{
+				const void* Src = Accessor.DataAt(0);
+				if (Accessor.ComponentType == FAccessor::EComponentType::F32)
+				{
+					memcpy(Buffer, Src, Accessor.Count * sizeof(ItemType));
+				}
+				else if (Accessor.Normalized)
+				{
+					CopyNormalized<ItemType, ItemElementCount>(Buffer, Src, Accessor.ComponentType, Accessor.Count);
+				}
+			}
+			else
+			{
+				if (Accessor.ComponentType == FAccessor::EComponentType::F32)
+				{
+					for (uint32 Index = 0; Index < Accessor.Count; ++Index)
+					{
+						const void* Pointer = Accessor.DataAt(Index);
+						Buffer[Index] = *static_cast<const ItemType*>(Pointer);
+					}
+				}
+				else if (Accessor.Normalized)
+				{
+					for (uint32 Index = 0; Index < Accessor.Count; ++Index)
+					{
+						const void* Pointer = Accessor.DataAt(Index);
+						Buffer[Index] = GetNormalized<ItemType, ItemElementCount>(Accessor.ComponentType, Pointer);
+					}
+				}
+			}
+		}
 	}
 
 	FAccessor::FAccessor(uint32 InCount, EType InType, EComponentType InCompType, bool InNormalized)
@@ -483,65 +522,29 @@ namespace GLTF
 
 	void FValidAccessor::GetVec2Array(FVector2f* Buffer) const
 	{
-		if (Type == EType::Vec2)  // strict format match, unlike GPU shader fetch
+		if (!ensure(Type == EType::Vec2))
 		{
-			const void* Src = DataAt(0);
-
-			if (ComponentType == EComponentType::F32)
-			{
-				// copy float vec2 directly from buffer
-				memcpy(Buffer, Src, Count * sizeof(FVector2f));
-				return;
-			}
-			else if (Normalized)
-			{
-				CopyNormalized<FVector2f, 2>(Buffer, Src, ComponentType, Count);
-				return;
-			}
+			return;
 		}
-		check(false);
+		CopyWithoutConversion<FVector2f, 2>(*this, Buffer);
 	}
 
 	void FValidAccessor::GetVec3Array(FVector3f* Buffer) const
 	{
-		if (Type == EType::Vec3)  // strict format match, unlike GPU shader fetch
+		if (!ensure(Type == EType::Vec3))
 		{
-			const void* Src = DataAt(0);
-
-			if (ComponentType == EComponentType::F32)
-			{
-				// copy float vec3 directly from buffer
-				memcpy(Buffer, Src, Count * sizeof(FVector3f));
-				return;
-			}
-			else if (Normalized)
-			{
-				CopyNormalized<FVector3f, 3>(Buffer, Src, ComponentType, Count);
-				return;
-			}
+			return;
 		}
-		check(false);
+		CopyWithoutConversion<FVector3f, 3>(*this, Buffer);
 	}
 
 	void FValidAccessor::GetVec4Array(FVector4f* Buffer) const
 	{
-		if (Type == EType::Vec4)  // strict format match, unlike GPU shader fetch
+		if (!ensure(Type == EType::Vec4))
 		{
-			const void* Src = DataAt(0);
-
-			if (ComponentType == EComponentType::F32)
-			{
-				// copy float vec4 directly from buffer
-				memcpy(Buffer, Src, Count * sizeof(FVector4f));
-				return;
-			}
-			else if (Normalized)
-			{
-				CopyNormalized<FVector4f, 4>(Buffer, Src, ComponentType, Count);
-				return;
-			}
+			return;
 		}
-		check(false);
+		CopyWithoutConversion<FVector4f, 4>(*this, Buffer);
 	}
 
 	void FValidAccessor::GetMat4Array(FMatrix44f* Buffer) const
diff --git a/Engine/Plugins/Editor/GLTFImporter/Source/GLTFCore/Public/GLTFAccessor.h b/Engine/Plugins/Editor/GLTFImporter/Source/GLTFCore/Public/GLTFAccessor.h
--- a/Engine/Plugins/Editor/GLTFImporter/Source/GLTFCore/Public/GLTFAccessor.h
+++ b/Engine/Plugins/Editor/GLTFImporter/Source/GLTFCore/Public/GLTFAccessor.h
@@ -140,13 +140,12 @@ namespace GLTF
 		void GetVec4Array(FVector4f* Buffer) const override;
 		void GetMat4Array(FMatrix44f* Buffer) const override;
 
-	private:
+		const uint8* DataAt(uint32 Index) const;
+
 		const FBufferView& BufferView;
 		const uint32       ByteOffset;
 		const uint32       ElementSize;
 		const uint32	   ByteStride;
-
-		const uint8* DataAt(uint32 Index) const;
 	};
 
 	struct GLTFCORE_API FVoidAccessor final : FAccessor
