Moved the GLTF import options to the public folder.

#9194

diff --git a/Engine/Plugins/Editor/GLTFImporter/Source/GLTFImporter/Private/GLTFImportOptions.h b/Engine/Plugins/Editor/GLTFImporter/Source/GLTFImporter/Private/GLTFImportOptions.h
deleted file mode 100644
diff --git a/Engine/Plugins/Editor/GLTFImporter/Source/GLTFImporter/Public/GLTFImportOptions.h b/Engine/Plugins/Editor/GLTFImporter/Source/GLTFImporter/Public/GLTFImportOptions.h
new file mode 100644
--- /dev/null
+++ b/Engine/Plugins/Editor/GLTFImporter/Source/GLTFImporter/Public/GLTFImportOptions.h
@@ -0,0 +1,33 @@
+// Copyright Epic Games, Inc. All Rights Reserved.
+
+#pragma once
+
+#include "CoreMinimal.h"
+#include "Engine/EngineTypes.h"
+#include "UObject/Object.h"
+#include "UObject/ObjectMacros.h"
+
+#include "GLTFImportOptions.generated.h"
+
+UCLASS(config = EditorPerProjectUserSettings, HideCategories = (DebugProperty))
+class GLTFIMPORTER_API UGLTFImportOptions : public UObject
+{
+	GENERATED_UCLASS_BODY()
+
+public:
+	UPROPERTY(
+		BlueprintReadWrite,
+	    config, EditAnywhere, Category = Lightmaps,
+	    meta =
+	        (ToolTip =
+	             "Generate new UV coordinates for lightmapping instead of using the highest index UV set. \nTurn this on to have Unreal Studio generate lightmap UV sets automatically.\nTurn this off to try using the highest index existing UV set (if available) as the lightmap UV set.\nFor both cases, geometry without existing UV sets will receive an empty UV set, which will by itself not be valid for use with Lightmass."))
+	bool bGenerateLightmapUVs;
+
+	UPROPERTY(
+		BlueprintReadWrite,
+	    config, EditAnywhere, Category = AssetImporting,
+	    meta =
+	        (DisplayName = "Import Uniform Scale",
+	         ToolTip = "Scale factor used for importing assets, by default: 100, for conversion from meters(glTF) to centimeters(Unreal default)."))
+	float ImportScale;
+};
