GLTF: Fixed a crash when a GLTF file doesn't have its accompanying bin file. The buffer accessor would be null but we'd still try to access its memory.

diff --git a/Engine/Plugins/Editor/GLTFImporter/Source/GLTFCore/Private/GLTF/GLTFAccessor.cpp b/Engine/Plugins/Editor/GLTFImporter/Source/GLTFCore/Private/GLTF/GLTFAccessor.cpp
--- a/Engine/Plugins/Editor/GLTFImporter/Source/GLTFCore/Private/GLTF/GLTFAccessor.cpp
+++ b/Engine/Plugins/Editor/GLTFImporter/Source/GLTFCore/Private/GLTF/GLTFAccessor.cpp
@@ -248,7 +248,7 @@ namespace GLTF
 
 	bool FValidAccessor::IsValid() const
 	{
-		return true;
+		return BufferView.IsValid();
 	}
 
 	FMD5Hash FValidAccessor::GetHash() const
diff --git a/Engine/Plugins/Editor/GLTFImporter/Source/GLTFCore/Private/GLTFMeshFactory.cpp b/Engine/Plugins/Editor/GLTFImporter/Source/GLTFCore/Private/GLTFMeshFactory.cpp
--- a/Engine/Plugins/Editor/GLTFImporter/Source/GLTFCore/Private/GLTFMeshFactory.cpp
+++ b/Engine/Plugins/Editor/GLTFImporter/Source/GLTFCore/Private/GLTFMeshFactory.cpp
@@ -181,6 +181,13 @@ namespace GLTF
 		for (int32 Index = 0; Index < Mesh.Primitives.Num(); ++Index)
 		{
 			const FPrimitive& Primitive = Mesh.Primitives[Index];
+
+			if (!Primitive.IsValid())
+			{
+				Messages.Emplace(EMessageSeverity::Warning, TEXT("Mesh has an invalid primitive: ") + Mesh.Name);
+				continue;
+			}
+
 			// Remember which primitives use which materials.
 			MaterialIndicesUsed.Add(Primitive.MaterialIndex);
 
@@ -214,6 +221,12 @@ namespace GLTF
 		for (int32 Index = 0; Index < Mesh.Primitives.Num(); ++Index)
 		{
 			const FPrimitive& Primitive = Mesh.Primitives[Index];
+
+			if (!Primitive.IsValid())
+			{
+				continue;
+			}
+
 			const bool        bHasDegenerateTriangles =
 				ImportPrimitive(Primitive, Index, NumUVs, Mesh.HasTangents(), Mesh.HasColors(),  //
 					VertexInstanceNormals, VertexInstanceTangents, VertexInstanceBinormalSigns, VertexInstanceUVs,
diff --git a/Engine/Plugins/Editor/GLTFImporter/Source/GLTFCore/Public/GLTFAccessor.h b/Engine/Plugins/Editor/GLTFImporter/Source/GLTFCore/Public/GLTFAccessor.h
--- a/Engine/Plugins/Editor/GLTFImporter/Source/GLTFCore/Public/GLTFAccessor.h
+++ b/Engine/Plugins/Editor/GLTFImporter/Source/GLTFCore/Public/GLTFAccessor.h
@@ -13,14 +13,20 @@ namespace GLTF
 		const uint32 ByteLength;
 		const uint8* Data;
 
-		FBuffer(uint32 InByteLength)
+		explicit FBuffer(uint32 InByteLength)
 		    : ByteLength(InByteLength)
 		    , Data(nullptr)
 		{
 		}
 
+		bool IsValid() const
+		{
+			return Data != nullptr;
+		}
+
 		const uint8* DataAt(uint32 Offset) const
 		{
+			checkSlow(Data);
 			return Data + Offset;
 		}
 	};
@@ -33,7 +39,7 @@ namespace GLTF
 		// if zero then accessor elements are tightly packed, i.e., effective stride equals the size of the element
 		const uint32 ByteStride;  // range 4..252
 
-		FBufferView(const FBuffer& InBuffer, uint32 InOffset, uint32 InLength, uint32 InStride)
+		explicit FBufferView(const FBuffer& InBuffer, uint32 InOffset, uint32 InLength, uint32 InStride)
 		    : Buffer(InBuffer)
 		    , ByteOffset(InOffset)
 		    , ByteLength(InLength)
@@ -42,6 +48,11 @@ namespace GLTF
 			// check that view fits completely inside the buffer
 		}
 
+		bool IsValid() const
+		{
+			return Buffer.IsValid();
+		}
+
 		const uint8* DataAt(uint32 Offset) const
 		{
 			return Buffer.DataAt(Offset + ByteOffset);
