Introduced Tangent acquisition skip into GLTFMeshFactory.
Interchange GLTF Import - the Gltf provided Tangents seem to be not usable by UE, because of that, removed/skipped the provided tangents in the Interchange Gltf Translator

diff --git a/Engine/Plugins/Editor/GLTFImporter/Source/GLTFCore/Private/GLTFMeshFactory.cpp b/Engine/Plugins/Editor/GLTFImporter/Source/GLTFCore/Private/GLTFMeshFactory.cpp
--- a/Engine/Plugins/Editor/GLTFImporter/Source/GLTFCore/Private/GLTFMeshFactory.cpp
+++ b/Engine/Plugins/Editor/GLTFImporter/Source/GLTFCore/Private/GLTFMeshFactory.cpp
@@ -16,7 +16,7 @@ namespace GLTF
 	public:
 		FMeshFactoryImpl();
 
-		void FillMeshDescription(const GLTF::FMesh &Mesh, FMeshDescription* MeshDescription);
+		void FillMeshDescription(const GLTF::FMesh &Mesh, FMeshDescription* MeshDescription, bool bSkipTangents);
 
 		void CleanUp();
 
@@ -33,7 +33,8 @@ namespace GLTF
 			const TVertexInstanceAttributesRef<FVector2f>& VertexInstanceUVs,
 			const TVertexInstanceAttributesRef<FVector4f>&  VertexInstanceColors,
 			const TEdgeAttributesRef<bool>&                EdgeHardnesses,
-			FMeshDescription* MeshDescription);
+			FMeshDescription* MeshDescription,
+			bool bSkipTangents);
 
 		inline TArray<FVector4f>& GetVector4dBuffer(int32 Index)
 		{
@@ -158,7 +159,7 @@ namespace GLTF
 		CornerVertexInstanceIDs.SetNum(3);
 	}
 
-	void FMeshFactoryImpl::FillMeshDescription(const FMesh &Mesh, FMeshDescription* MeshDescription)
+	void FMeshFactoryImpl::FillMeshDescription(const FMesh &Mesh, FMeshDescription* MeshDescription, bool bSkipTangents)
 	{
 		const int32 NumUVs = FMath::Max(1, GetNumUVs(Mesh));
 
@@ -231,7 +232,7 @@ namespace GLTF
 				ImportPrimitive(Primitive, Index, NumUVs, Mesh.HasTangents(), Mesh.HasColors(),  //
 					VertexInstanceNormals, VertexInstanceTangents, VertexInstanceBinormalSigns, VertexInstanceUVs,
 					VertexInstanceColors,  //
-					EdgeHardnesses, MeshDescription);
+					EdgeHardnesses, MeshDescription, bSkipTangents);
 
 			bMeshUsesEmptyMaterial |= Primitive.MaterialIndex == INDEX_NONE;
 			for (int32 UVIndex = 0; UVIndex < NumUVs; ++UVIndex)
@@ -266,7 +267,8 @@ namespace GLTF
 		const TVertexInstanceAttributesRef<FVector2f>& VertexInstanceUVs,
 		const TVertexInstanceAttributesRef<FVector4f>&  VertexInstanceColors,
 		const TEdgeAttributesRef<bool>&                EdgeHardnesses,
-		FMeshDescription* MeshDescription)
+		FMeshDescription* MeshDescription,
+		bool bSkipTangents)
 	{
 
 		const FPolygonGroupID CurrentPolygonGroupID = MaterialIndexToPolygonGroupID[Primitive.MaterialIndex];
@@ -374,17 +376,17 @@ namespace GLTF
 
 				const uint32 IndiceIndex = TriangleIndex * 3 + Corner;
 
-				if (Tangents.Num() > 0)
+				VertexInstanceNormals[VertexInstanceID] = Normals[IndiceIndex];
+
+				if (!bSkipTangents && Tangents.Num() > 0)
 				{
 					VertexInstanceTangents[VertexInstanceID] = Tangents[IndiceIndex];
+					VertexInstanceBinormalSigns[VertexInstanceID] =
+						GetBasisDeterminantSign((FVector)VertexInstanceTangents[VertexInstanceID].GetSafeNormal(),
+							(FVector)(VertexInstanceNormals[VertexInstanceID] ^ VertexInstanceTangents[VertexInstanceID]).GetSafeNormal(),
+							(FVector)VertexInstanceNormals[VertexInstanceID].GetSafeNormal());
 				}
 
-				VertexInstanceNormals[VertexInstanceID] = Normals[IndiceIndex];
-				VertexInstanceBinormalSigns[VertexInstanceID] =
-					GetBasisDeterminantSign((FVector)VertexInstanceTangents[VertexInstanceID].GetSafeNormal(),
-						(FVector)(VertexInstanceNormals[VertexInstanceID] ^ VertexInstanceTangents[VertexInstanceID]).GetSafeNormal(),
-						(FVector)VertexInstanceNormals[VertexInstanceID].GetSafeNormal());
-
 				for (int32 UVIndex = 0; UVIndex < NumUVs; ++UVIndex)
 				{
 					VertexInstanceUVs.Set(VertexInstanceID, UVIndex, (*UVs[UVIndex])[IndiceIndex]);
@@ -442,9 +444,9 @@ namespace GLTF
 
 	FMeshFactory::~FMeshFactory() {}
 
-	void FMeshFactory::FillMeshDescription(const GLTF::FMesh &Mesh, FMeshDescription* MeshDescription)
+	void FMeshFactory::FillMeshDescription(const GLTF::FMesh &Mesh, FMeshDescription* MeshDescription, bool bSkipTangents)
 	{
-		Impl->FillMeshDescription(Mesh, MeshDescription);
+		Impl->FillMeshDescription(Mesh, MeshDescription, bSkipTangents);
 	}
 
 	const TArray<FLogMessage>& FMeshFactory::GetLogMessages() const
diff --git a/Engine/Plugins/Editor/GLTFImporter/Source/GLTFCore/Public/GLTFMeshFactory.h b/Engine/Plugins/Editor/GLTFImporter/Source/GLTFCore/Public/GLTFMeshFactory.h
--- a/Engine/Plugins/Editor/GLTFImporter/Source/GLTFCore/Public/GLTFMeshFactory.h
+++ b/Engine/Plugins/Editor/GLTFImporter/Source/GLTFCore/Public/GLTFMeshFactory.h
@@ -18,7 +18,7 @@ namespace GLTF
 		FMeshFactory();
 		~FMeshFactory();
 
-		void FillMeshDescription(const GLTF::FMesh &Mesh, FMeshDescription* MeshDescription);
+		void FillMeshDescription(const GLTF::FMesh &Mesh, FMeshDescription* MeshDescription, bool bSkipTangents = false);
 
 		float GetUniformScale() const;
 		void  SetUniformScale(float Scale);
diff --git a/Engine/Plugins/Experimental/Interchange/Runtime/Source/Import/Private/Gltf/InterchangeGltfTranslator.cpp b/Engine/Plugins/Experimental/Interchange/Runtime/Source/Import/Private/Gltf/InterchangeGltfTranslator.cpp
--- a/Engine/Plugins/Experimental/Interchange/Runtime/Source/Import/Private/Gltf/InterchangeGltfTranslator.cpp
+++ b/Engine/Plugins/Experimental/Interchange/Runtime/Source/Import/Private/Gltf/InterchangeGltfTranslator.cpp
@@ -940,7 +940,7 @@ TFuture< TOptional< UE::Interchange::FStaticMeshPayloadData > > UInterchangeGltf
 	const GLTF::FMesh& GltfMesh = GltfAsset.Meshes[ MeshIndex ];
 	GLTF::FMeshFactory MeshFactory;
 	MeshFactory.SetUniformScale( 100.f ); // GLTF is in meters while UE is in centimeters
-	MeshFactory.FillMeshDescription( GltfMesh, &StaticMeshPayloadData.MeshDescription );
+	MeshFactory.FillMeshDescription( GltfMesh, &StaticMeshPayloadData.MeshDescription, true );
 
 	// Patch polygon groups material slot names to match Interchange expectations (rename material slots from indices to material names)
 	{
