Transparent/Translucent GLTF material black on import.
Removing the possibility of KHR_MaterialTransmission if Material is alphaMode.Blend.

diff --git a/Engine/Plugins/Editor/GLTFImporter/Source/GLTFCore/Private/GLTF/ExtensionsHandler.cpp b/Engine/Plugins/Editor/GLTFImporter/Source/GLTFCore/Private/GLTF/ExtensionsHandler.cpp
--- a/Engine/Plugins/Editor/GLTFImporter/Source/GLTFCore/Private/GLTF/ExtensionsHandler.cpp
+++ b/Engine/Plugins/Editor/GLTFImporter/Source/GLTFCore/Private/GLTF/ExtensionsHandler.cpp
@@ -173,6 +173,10 @@ namespace GLTF
 				break;
 				case KHR_materials_transmission:
 				{
+					if (Material.AlphaMode == FMaterial::EAlphaMode::Blend)
+					{
+						break;
+					}
 					const FJsonObject& Transm = ExtObj;
 
 					Material.bHasTransmission = true;
