GLTF tangents fix
- Take into account the tangent space handedness (tangent w component). This is used as the bitangent sign in mesh description. Flipping the bitangent vector is equivalent to changing the handedness.
- Note that our bitangents are flipped compared to the GLTF expected tangent space but this is compensated by flipping the green channel when sampling the normal map. This tangent space is similar to the one UE would auto generate so filliping the normal map green channel allows it to map to both GLTF provided tangents and UE generated tangents.

diff --git a/Engine/Plugins/Editor/GLTFImporter/Source/GLTFCore/Private/GLTF/ConversionUtilities.h b/Engine/Plugins/Editor/GLTFImporter/Source/GLTFCore/Private/GLTF/ConversionUtilities.h
--- a/Engine/Plugins/Editor/GLTFImporter/Source/GLTFCore/Private/GLTF/ConversionUtilities.h
+++ b/Engine/Plugins/Editor/GLTFImporter/Source/GLTFCore/Private/GLTF/ConversionUtilities.h
@@ -21,12 +21,11 @@ namespace GLTF
 		return {Vec.X, Vec.Z, Vec.Y};
 	}
 
-	inline FVector ConvertTangent(const FVector4& Tangent)
+	inline FVector4 ConvertTangent(const FVector4& Tangent)
 	{
 		// glTF stores tangent as Vec4, with W component indicating handedness of tangent basis.
 
-		FVector Result(Tangent.X, Tangent.Y, Tangent.Z);
-		Result = ConvertVec3(Result);
+		FVector4 Result(ConvertVec3(FVector(Tangent.X, Tangent.Y, Tangent.Z)), Tangent.W);
 		return Result;
 	}
 
diff --git a/Engine/Plugins/Editor/GLTFImporter/Source/GLTFCore/Private/GLTF/GLTFMesh.cpp b/Engine/Plugins/Editor/GLTFImporter/Source/GLTFCore/Private/GLTF/GLTFMesh.cpp
--- a/Engine/Plugins/Editor/GLTFImporter/Source/GLTFCore/Private/GLTF/GLTFMesh.cpp
+++ b/Engine/Plugins/Editor/GLTFImporter/Source/GLTFCore/Private/GLTF/GLTFMesh.cpp
@@ -36,12 +36,12 @@ namespace GLTF
 	{
 	}
 
-	void FPrimitive::GetTangents(TArray<FVector3f>& Buffer) const
+	void FPrimitive::GetTangents(TArray<FVector4f>& Buffer) const
 	{
 		Buffer.Reserve(Tangent.Count);
 		for (uint32 Index = 0; Index < Tangent.Count; ++Index)
 		{
-			Buffer.Push((FVector3f)GLTF::ConvertTangent(Tangent.GetVec4(Index)));
+			Buffer.Push((FVector4f)GLTF::ConvertTangent(Tangent.GetVec4(Index)));
 		}
 	}
 
diff --git a/Engine/Plugins/Editor/GLTFImporter/Source/GLTFCore/Private/GLTFMeshFactory.cpp b/Engine/Plugins/Editor/GLTFImporter/Source/GLTFCore/Private/GLTFMeshFactory.cpp
--- a/Engine/Plugins/Editor/GLTFImporter/Source/GLTFCore/Private/GLTFMeshFactory.cpp
+++ b/Engine/Plugins/Editor/GLTFImporter/Source/GLTFCore/Private/GLTFMeshFactory.cpp
@@ -16,7 +16,7 @@ namespace GLTF
 	public:
 		FMeshFactoryImpl();
 
-		void FillMeshDescription(const GLTF::FMesh &Mesh, FMeshDescription* MeshDescription, bool bSkipTangents);
+		void FillMeshDescription(const GLTF::FMesh &Mesh, FMeshDescription* MeshDescription);
 
 		void CleanUp();
 
@@ -71,14 +71,14 @@ namespace GLTF
 		enum
 		{
 			NormalBufferIndex = 0,
-			TangentBufferIndex = 1,
-			PositionBufferIndex = 2,
-			ReindexBufferIndex = 3,
-			VectorBufferCount = 4,
+			PositionBufferIndex = 1,
+			ReindexBufferIndex = 2,
+			VectorBufferCount = 3,
 			UvReindexBufferIndex = MAX_MESH_TEXTURE_COORDS_MD,
 			ColorBufferIndex = 0,
-			Reindex4dBufferIndex = 1,
-			Vector4dBufferCount = 2,
+			TangentBufferIndex = 1,
+			Reindex4dBufferIndex = 2,
+			Vector4dBufferCount = 3,
 		};
 
 		using FIndexVertexIdMap = TMap<int32, FVertexID>;
@@ -159,8 +159,10 @@ namespace GLTF
 		CornerVertexInstanceIDs.SetNum(3);
 	}
 
-	void FMeshFactoryImpl::FillMeshDescription(const FMesh &Mesh, FMeshDescription* MeshDescription, bool bSkipTangents)
+	void FMeshFactoryImpl::FillMeshDescription(const FMesh &Mesh, FMeshDescription* MeshDescription)
 	{
+		const bool bSkipTangents = !Mesh.HasNormals(); // Per the GLTF spec, tangents should be ignored if no normals are provided
+
 		const int32 NumUVs = FMath::Max(1, GetNumUVs(Mesh));
 
 		FStaticMeshAttributes StaticMeshAttributes(*MeshDescription);
@@ -259,7 +261,7 @@ namespace GLTF
 	bool FMeshFactoryImpl::ImportPrimitive(const GLTF::FPrimitive&                        Primitive,  //
 		int32                                          PrimitiveIndex,
 		int32                                          NumUVs,
-		bool                                           bMeshHasTagents,
+		bool                                           bMeshHasTangents,
 		bool                                           bMeshHasColors,
 		const TVertexInstanceAttributesRef<FVector3f>&   VertexInstanceNormals,
 		const TVertexInstanceAttributesRef<FVector3f>&   VertexInstanceTangents,
@@ -297,18 +299,18 @@ namespace GLTF
 			GenerateFlatNormals(Positions, Indices, Normals);
 		}
 
-		TArray<FVector3f>& Tangents = GetVectorBuffer(TangentBufferIndex);
+		TArray<FVector4f>& Tangents = GetVector4dBuffer(TangentBufferIndex);
 		if (Primitive.HasTangents())
 		{
-			TArray<FVector3f>& ReindexBuffer = GetVectorBuffer(ReindexBufferIndex);
+			TArray<FVector4f>& ReindexBuffer = GetVector4dBuffer(Reindex4dBufferIndex);
 			Primitive.GetTangents(Tangents);
 			ReIndex(Tangents, Indices, ReindexBuffer);
 			Swap(Tangents, ReindexBuffer);
 		}
-		else if (bMeshHasTagents)
+		else if (bMeshHasTangents)
 		{
 			// If other primitives in this mesh have tangents, generate filler ones for this primitive, to avoid gaps.
-			Tangents.Init(FVector3f(0.0f, 0.0f, 1.0f), Primitive.VertexCount());
+			Tangents.Init(FVector4f(1.0f, 0.0f, 0.0f, 1.0f), Primitive.VertexCount());
 		}
 
 		TArray<FVector4f>& Colors = GetVector4dBuffer(ColorBufferIndex);
@@ -381,10 +383,7 @@ namespace GLTF
 				if (!bSkipTangents && Tangents.Num() > 0)
 				{
 					VertexInstanceTangents[VertexInstanceID] = Tangents[IndiceIndex];
-					VertexInstanceBinormalSigns[VertexInstanceID] =
-						GetBasisDeterminantSign((FVector)VertexInstanceTangents[VertexInstanceID].GetSafeNormal(),
-							(FVector)(VertexInstanceNormals[VertexInstanceID] ^ VertexInstanceTangents[VertexInstanceID]).GetSafeNormal(),
-							(FVector)VertexInstanceNormals[VertexInstanceID].GetSafeNormal());
+					VertexInstanceBinormalSigns[VertexInstanceID] = Tangents[IndiceIndex].W;
 				}
 
 				for (int32 UVIndex = 0; UVIndex < NumUVs; ++UVIndex)
@@ -444,9 +443,9 @@ namespace GLTF
 
 	FMeshFactory::~FMeshFactory() {}
 
-	void FMeshFactory::FillMeshDescription(const GLTF::FMesh &Mesh, FMeshDescription* MeshDescription, bool bSkipTangents)
+	void FMeshFactory::FillMeshDescription(const GLTF::FMesh &Mesh, FMeshDescription* MeshDescription)
 	{
-		Impl->FillMeshDescription(Mesh, MeshDescription, bSkipTangents);
+		Impl->FillMeshDescription(Mesh, MeshDescription);
 	}
 
 	const TArray<FLogMessage>& FMeshFactory::GetLogMessages() const
diff --git a/Engine/Plugins/Editor/GLTFImporter/Source/GLTFCore/Public/GLTFMesh.h b/Engine/Plugins/Editor/GLTFImporter/Source/GLTFCore/Public/GLTFMesh.h
--- a/Engine/Plugins/Editor/GLTFImporter/Source/GLTFCore/Public/GLTFMesh.h
+++ b/Engine/Plugins/Editor/GLTFImporter/Source/GLTFCore/Public/GLTFMesh.h
@@ -72,7 +72,7 @@ namespace GLTF
 		bool HasNormals() const;
 		void GetNormals(TArray<FVector3f>& Buffer) const;
 		bool HasTangents() const;
-		void GetTangents(TArray<FVector3f>& Buffer) const;
+		void GetTangents(TArray<FVector4f>& Buffer) const;
 		bool HasTexCoords(uint32 Index) const;
 		void GetTexCoords(uint32 Index, TArray<FVector2f>& Buffer) const;
 		void GetColors(TArray<FVector4f>& Buffer) const;
diff --git a/Engine/Plugins/Editor/GLTFImporter/Source/GLTFCore/Public/GLTFMeshFactory.h b/Engine/Plugins/Editor/GLTFImporter/Source/GLTFCore/Public/GLTFMeshFactory.h
--- a/Engine/Plugins/Editor/GLTFImporter/Source/GLTFCore/Public/GLTFMeshFactory.h
+++ b/Engine/Plugins/Editor/GLTFImporter/Source/GLTFCore/Public/GLTFMeshFactory.h
@@ -18,7 +18,7 @@ namespace GLTF
 		FMeshFactory();
 		~FMeshFactory();
 
-		void FillMeshDescription(const GLTF::FMesh &Mesh, FMeshDescription* MeshDescription, bool bSkipTangents = false);
+		void FillMeshDescription(const GLTF::FMesh &Mesh, FMeshDescription* MeshDescription);
 
 		float GetUniformScale() const;
 		void  SetUniformScale(float Scale);
