Fixing empty level sequence import for GLTF imports (adding Skeleton index to joints)

diff --git a/Engine/Plugins/Editor/GLTFImporter/Source/GLTFCore/Private/GLTF/GLTFReader.cpp b/Engine/Plugins/Editor/GLTFImporter/Source/GLTFCore/Private/GLTF/GLTFReader.cpp
--- a/Engine/Plugins/Editor/GLTFImporter/Source/GLTFCore/Private/GLTF/GLTFReader.cpp
+++ b/Engine/Plugins/Editor/GLTFImporter/Source/GLTFCore/Private/GLTF/GLTFReader.cpp
@@ -445,6 +445,26 @@ namespace GLTF
 		}
 
 		Skin.Skeleton = GetIndex(Object, TEXT("skeleton"));
+		if (Skin.Skeleton != INDEX_NONE)
+		{
+			//From gltf specification:
+			//If Skeleton is present:
+			//	"The index of the node used as a skeleton root."
+			//If it is not part of the Joints list yet, add it:
+			bool bSkeletonIsAlreadyJoint = false;
+			for (const int32& JointIndex : Skin.Joints)
+			{
+				if (JointIndex == Skin.Skeleton)
+				{
+					bSkeletonIsAlreadyJoint = true;
+					break;
+				}
+			}
+			if (!bSkeletonIsAlreadyJoint)
+			{
+				Skin.Joints.Add(Skin.Skeleton);
+			}
+		}
 
 		ExtensionsHandler->SetupSkinExtensions(Object, Skin);
 	}
