GLTF: Fixed error when trying to read the extensionsRequired field for files that don't have it.

diff --git a/Engine/Plugins/Editor/GLTFImporter/Source/GLTFCore/Private/GLTF/GLTFReader.cpp b/Engine/Plugins/Editor/GLTFImporter/Source/GLTFCore/Private/GLTF/GLTFReader.cpp
--- a/Engine/Plugins/Editor/GLTFImporter/Source/GLTFCore/Private/GLTF/GLTFReader.cpp
+++ b/Engine/Plugins/Editor/GLTFImporter/Source/GLTFCore/Private/GLTF/GLTFReader.cpp
@@ -880,17 +880,20 @@ namespace GLTF
 		SetupObjects(MaterialCount, TEXT("materials"), [this](const FJsonObject& Object) { SetupMaterial(Object); });
 
 
-		/*
-		* Returns true if all required extensions are supported.
-		* According to gltf specification checking only the top-level extensionsRequired property in order to decide
-		* if the model import is supported should be sufficient as the documentation states:
-		* "All glTF extensions required to load and/or render an asset MUST be listed in the top-level extensionsRequired array"
-		*/
-		{
-			for (const TSharedPtr<FJsonValue>& Extension : JsonRoot->GetArrayField(TEXT("extensionsRequired")))
-			{
-				FString ExtensionString = Extension->AsString();
-				OutAsset.RequiredExtensions.Add(ExtensionString);
+		/**
+		 * Returns true if all required extensions are supported.
+		 * According to gltf specification checking only the top-level extensionsRequired property in order to decide
+		 * if the model import is supported should be sufficient as the documentation states:
+		 * "All glTF extensions required to load and/or render an asset MUST be listed in the top-level extensionsRequired array"
+		 */
+		const TArray<TSharedPtr<FJsonValue>>* ExtensionsRequired;
+		if (JsonRoot->TryGetArrayField(TEXT("extensionsRequired"), ExtensionsRequired))
+		{
+			check(ExtensionsRequired);
+			OutAsset.RequiredExtensions.Reserve(OutAsset.RequiredExtensions.Num() + ExtensionsRequired->Num());
+			for (const TSharedPtr<FJsonValue>& Extension : *ExtensionsRequired)
+			{
+				OutAsset.RequiredExtensions.Add(Extension->AsString());
 			}
 		}
 
