Fixing Gltf Impoter embedded Texture recognition

diff --git a/Engine/Plugins/Editor/GLTFImporter/Source/GLTFCore/Private/GLTF/GLTFReader.cpp b/Engine/Plugins/Editor/GLTFImporter/Source/GLTFCore/Private/GLTF/GLTFReader.cpp
--- a/Engine/Plugins/Editor/GLTFImporter/Source/GLTFCore/Private/GLTF/GLTFReader.cpp
+++ b/Engine/Plugins/Editor/GLTFImporter/Source/GLTFCore/Private/GLTF/GLTFReader.cpp
@@ -481,7 +481,7 @@ namespace GLTF
 			// Unreal *is* responsible for decoding Data based on Format.
 
 			Image.URI = Object.GetStringField(TEXT("uri"));
-			if (Image.URI.StartsWith(TEXT("data:)")))
+			if (Image.URI.StartsWith(TEXT("data:")))
 			{
 				uint32  ImageSize = 0;
 				FString MimeType;
