[interchange] [gltf] Editor crash on import of gltf file with embedded base64-encoded texture fix

diff --git a/Engine/Plugins/Editor/GLTFImporter/Source/GLTFCore/Private/GLTF/GLTFReader.cpp b/Engine/Plugins/Editor/GLTFImporter/Source/GLTFCore/Private/GLTF/GLTFReader.cpp
--- a/Engine/Plugins/Editor/GLTFImporter/Source/GLTFCore/Private/GLTF/GLTFReader.cpp
+++ b/Engine/Plugins/Editor/GLTFImporter/Source/GLTFCore/Private/GLTF/GLTFReader.cpp
@@ -785,7 +785,7 @@ namespace GLTF
 					continue;
 
 				const FString& URI = Object.GetStringField(TEXT("uri"));
-				if (URI.StartsWith(TEXT("data:)")))
+				if (URI.StartsWith(TEXT("data:")))
 				{
 					FString         MimeType;
 					const uint32    DataSize = GetDecodedDataSize(URI, MimeType);
