[GLTFCore Importer] GLTF importer crash with array out of bounds fix

diff --git a/Engine/Plugins/Editor/GLTFImporter/Source/GLTFCore/Private/GLTFMeshFactory.cpp b/Engine/Plugins/Editor/GLTFImporter/Source/GLTFCore/Private/GLTFMeshFactory.cpp
--- a/Engine/Plugins/Editor/GLTFImporter/Source/GLTFCore/Private/GLTFMeshFactory.cpp
+++ b/Engine/Plugins/Editor/GLTFImporter/Source/GLTFCore/Private/GLTFMeshFactory.cpp
@@ -277,6 +277,22 @@ namespace GLTF
 		TArray<uint32>& Indices = GetIntBuffer();
 		Primitive.GetTriangleIndices(Indices);
 
+		// Validate Indices against Positions:
+		//  In case of corrupted Indices return with bHasDegenerateTriangles.
+		{
+			TArray<FVector3f>& Positions = GetVectorBuffer(PositionBufferIndex);
+			Primitive.GetPositions(Positions);
+			uint32 PositionsSize = Positions.Num();
+
+			for (uint32 Index : Indices)
+			{
+				if (PositionsSize <= Index)
+				{
+					return true;
+				}
+			}
+		}
+
 		TArray<FVector3f>& Normals = GetVectorBuffer(NormalBufferIndex);
 		// glTF does not guarantee each primitive within a mesh has the same attributes.
 		// Fill in gaps as needed:
