[Interchange] Crash on gltf file import due FName none

diff --git a/Engine/Plugins/Editor/GLTFImporter/Source/GLTFCore/Private/GLTF/GLTFAsset.cpp b/Engine/Plugins/Editor/GLTFImporter/Source/GLTFCore/Private/GLTF/GLTFAsset.cpp
--- a/Engine/Plugins/Editor/GLTFImporter/Source/GLTFCore/Private/GLTF/GLTFAsset.cpp
+++ b/Engine/Plugins/Editor/GLTFImporter/Source/GLTFCore/Private/GLTF/GLTFAsset.cpp
@@ -28,6 +28,27 @@ namespace GLTF
 				}
 			}
 		}
+
+		//Make the Name FName complient.
+		//For Scenario where Asset.Name is '_x' (where x is a number) making FName to be none:
+		//UE-174253
+		template <typename T>
+		void ValidateFNameCriteria(const FString& Prefix, TArray<T>& Objects)
+		{
+			for (int32 Index = 0; Index < Objects.Num(); ++Index)
+			{
+				T& Obj = Objects[Index];
+
+				if (!Obj.Name.IsEmpty())
+				{
+					FName SanitizationTest(*Obj.Name);
+					if (SanitizationTest.IsNone())
+					{
+						Obj.Name = Prefix + Obj.Name;
+					}
+				}
+			}
+		}
 	}
 
 	const FMetadata::FExtraData* FMetadata::GetExtraData(const FString& Name) const
@@ -203,6 +224,12 @@ namespace GLTF
 		GLTF::GenerateNames(Prefix + TEXT("_animation_"), Animations);
 		GLTF::GenerateNames(Prefix + TEXT("_image_"), Images);
 		GLTF::GenerateNames(Prefix + TEXT("_mesh_"), Meshes);
+
+		// Validate FName criteria for Materials Textures Meshes Animations (assets)
+		GLTF::ValidateFNameCriteria(TEXT("Material_"), Materials);
+		GLTF::ValidateFNameCriteria(TEXT("Texture_"), Textures);
+		GLTF::ValidateFNameCriteria(TEXT("Mesh_"), Meshes);
+		GLTF::ValidateFNameCriteria(TEXT("Animation_"), Animations);
 	}
 
 	void FAsset::GetRootNodes(TArray<int32>& NodeIndices)
