[GLTFCore] MorphTarget weight existance check fix for parsing

diff --git a/Engine/Plugins/Editor/GLTFImporter/Source/GLTFCore/Private/GLTF/GLTFReader.cpp b/Engine/Plugins/Editor/GLTFImporter/Source/GLTFCore/Private/GLTF/GLTFReader.cpp
--- a/Engine/Plugins/Editor/GLTFImporter/Source/GLTFCore/Private/GLTF/GLTFReader.cpp
+++ b/Engine/Plugins/Editor/GLTFImporter/Source/GLTFCore/Private/GLTF/GLTFReader.cpp
@@ -295,10 +295,14 @@ namespace GLTF
 		}
 
 		// Morph Target Weights:
-		const TArray<TSharedPtr<FJsonValue> >& MorphTargetWeightsArray = Object.GetArrayField(TEXT("weights"));
-		for (TSharedPtr<FJsonValue> Value : MorphTargetWeightsArray)
+		if (Object.HasField(TEXT("weights")))
+
 		{
-			Mesh.MorphTargetWeights.Add(Value->AsNumber());
+			const TArray<TSharedPtr<FJsonValue> >& MorphTargetWeightsArray = Object.GetArrayField(TEXT("weights"));
+			for (TSharedPtr<FJsonValue> Value : MorphTargetWeightsArray)
+			{
+				Mesh.MorphTargetWeights.Add(Value->AsNumber());
+			}
 		}
 
 		// Morph Target Names:
