[Interchange] fails at loading textures if they have space in name fixes

diff --git a/Engine/Plugins/Editor/GLTFImporter/Source/GLTFCore/GLTFCore.Build.cs b/Engine/Plugins/Editor/GLTFImporter/Source/GLTFCore/GLTFCore.Build.cs
--- a/Engine/Plugins/Editor/GLTFImporter/Source/GLTFCore/GLTFCore.Build.cs
+++ b/Engine/Plugins/Editor/GLTFImporter/Source/GLTFCore/GLTFCore.Build.cs
@@ -22,6 +22,7 @@ namespace UnrealBuildTool.Rules
 					"StaticMeshDescription",
                     "Json",
                     "RenderCore",
+					"HTTP",
                 }
                 );
         }
diff --git a/Engine/Plugins/Editor/GLTFImporter/Source/GLTFCore/Private/GLTF/GLTFReader.cpp b/Engine/Plugins/Editor/GLTFImporter/Source/GLTFCore/Private/GLTF/GLTFReader.cpp
--- a/Engine/Plugins/Editor/GLTFImporter/Source/GLTFCore/Private/GLTF/GLTFReader.cpp
+++ b/Engine/Plugins/Editor/GLTFImporter/Source/GLTFCore/Private/GLTF/GLTFReader.cpp
@@ -17,6 +17,8 @@
 #include "Misc/Paths.h"
 #include "Serialization/JsonSerializer.h"
 
+#include "GenericPlatform/GenericPlatformHttp.h"
+
 namespace GLTF
 {
 	namespace
@@ -119,7 +121,7 @@ namespace GLTF
 		{
 			// set Buffer.Data from Object.uri
 
-			const FString& URI = Object.GetStringField(TEXT("uri"));
+			FString URI = Object.GetStringField(TEXT("uri"));
 			if (URI.StartsWith(TEXT("data:")))
 			{
 				FString MimeType;
@@ -136,6 +138,8 @@ namespace GLTF
 			}
 			else
 			{
+				URI = FGenericPlatformHttp::UrlDecode(URI);
+
 				// Load buffer from external file.
 				const FString FullPath = Path / URI;
 				FArchive*     Reader   = IFileManager::Get().CreateFileReader(*FullPath);
@@ -570,6 +574,7 @@ namespace GLTF
 			}
 			else  // Load buffer from external file.
 			{
+				Image.URI = FGenericPlatformHttp::UrlDecode(Image.URI);
 				Image.Format = ImageFormatFromFilename(Image.URI);
 
 				Image.FilePath = Path / Image.URI;
@@ -824,7 +829,7 @@ namespace GLTF
 				if (!Object.HasTypedField<EJson::String>(TEXT("uri")))
 					continue;
 
-				const FString& URI = Object.GetStringField(TEXT("uri"));
+				FString URI = Object.GetStringField(TEXT("uri"));
 				if (URI.StartsWith(TEXT("data:")))
 				{
 					FString      MimeType;
@@ -837,6 +842,7 @@ namespace GLTF
 				}
 				else
 				{
+					URI = FGenericPlatformHttp::UrlDecode(URI);
 					const FString FullPath = InResourcesPath / URI;
 					const int64   FileSize = FPlatformFileManager::Get().GetPlatformFile().FileSize(*FullPath);
 					if (ByteLength == FileSize)
@@ -856,7 +862,7 @@ namespace GLTF
 				if (!Object.HasTypedField<EJson::String>(TEXT("uri")))
 					continue;
 
-				const FString& URI = Object.GetStringField(TEXT("uri"));
+				FString URI = Object.GetStringField(TEXT("uri"));
 				if (URI.StartsWith(TEXT("data:")))
 				{
 					FString         MimeType;
@@ -869,6 +875,8 @@ namespace GLTF
 				}
 				else if (bInLoadImageData)
 				{
+					URI = FGenericPlatformHttp::UrlDecode(URI);
+
 					FImage::EFormat Format = ImageFormatFromFilename(URI);
 					if (Format != FImage::EFormat::Unknown)
 					{
