[GLTFCore/Interchange] GLTF Extension usage consinstency fixes, Interchange GLTF Analytics fixes.
Moved check() to ensure() wherever possible in GLTFImporter

diff --git a/Engine/Plugins/Editor/GLTFImporter/Source/GLTFCore/Private/GLTF/ConversionUtilities.h b/Engine/Plugins/Editor/GLTFImporter/Source/GLTFCore/Private/GLTF/ConversionUtilities.h
--- a/Engine/Plugins/Editor/GLTFImporter/Source/GLTFCore/Private/GLTF/ConversionUtilities.h
+++ b/Engine/Plugins/Editor/GLTFImporter/Source/GLTFCore/Private/GLTF/ConversionUtilities.h
@@ -78,7 +78,7 @@ namespace GLTF
 		{
 			return FMaterial::EAlphaMode::Blend;
 		}
-		check(false);
+		ensure(false);
 		return FMaterial::EAlphaMode::Opaque;
 	}
 
@@ -101,7 +101,7 @@ namespace GLTF
 		{
 			return FAnimation::EPath::Weights;
 		}
-		check(false);
+		ensure(false);
 		return FAnimation::EPath::Rotation;
 	}
 
@@ -116,7 +116,11 @@ namespace GLTF
 		    FPrimitive::EMode::TriangleStrip,
 		    FPrimitive::EMode::TriangleFan,
 		});
-		check(SafetyCheck.Find(static_cast<FPrimitive::EMode>(Num)) != INDEX_NONE);
+
+		if (!ensure(SafetyCheck.Find(static_cast<FPrimitive::EMode>(Num)) != INDEX_NONE))
+		{
+			return FPrimitive::EMode::Unknown;
+		}
 
 		return static_cast<FPrimitive::EMode>(Num);
 	}
@@ -156,7 +160,10 @@ namespace GLTF
 		static const TArray<FSampler::EFilter> SafetyCheck({FSampler::EFilter::Nearest, FSampler::EFilter::Linear,
 		                                                    FSampler::EFilter::NearestMipmapNearest, FSampler::EFilter::LinearMipmapNearest,
 		                                                    FSampler::EFilter::NearestMipmapLinear, FSampler::EFilter::LinearMipmapLinear});
-		check(SafetyCheck.Find(static_cast<FSampler::EFilter>(Num)) != INDEX_NONE);
+		if (!ensure(SafetyCheck.Find(static_cast<FSampler::EFilter>(Num)) != INDEX_NONE))
+		{
+			return FSampler::EFilter::Linear;
+		}
 
 		return static_cast<FSampler::EFilter>(Num);
 	}
@@ -164,7 +171,10 @@ namespace GLTF
 	inline FSampler::EWrap WrapModeFromNumber(uint16 Num)
 	{
 		static const TArray<FSampler::EWrap> SafetyCheck({FSampler::EWrap::ClampToEdge, FSampler::EWrap::MirroredRepeat, FSampler::EWrap::Repeat});
-		check(SafetyCheck.Find(static_cast<FSampler::EWrap>(Num)) != INDEX_NONE);
+		if (!ensure(SafetyCheck.Find(static_cast<FSampler::EWrap>(Num)) != INDEX_NONE))
+		{
+			return FSampler::EWrap::Repeat;
+		}
 
 		return static_cast<FSampler::EWrap>(Num);
 	}
diff --git a/Engine/Plugins/Editor/GLTFImporter/Source/GLTFCore/Private/GLTF/ExtensionsHandler.cpp b/Engine/Plugins/Editor/GLTFImporter/Source/GLTFCore/Private/GLTF/ExtensionsHandler.cpp
--- a/Engine/Plugins/Editor/GLTFImporter/Source/GLTFCore/Private/GLTF/ExtensionsHandler.cpp
+++ b/Engine/Plugins/Editor/GLTFImporter/Source/GLTFCore/Private/GLTF/ExtensionsHandler.cpp
@@ -11,8 +11,8 @@ namespace GLTF
 {
 	namespace
 	{
-		static const TArray<FString> LightExtensions = {TEXT("KHR_lights_punctual"), TEXT("KHR_lights")};
-		static const TArray<FString> MaterialsExtensions = {TEXT("KHR_materials_variants")};
+		static const TArray<FString> LightExtensions = { GLTF::ToString(GLTF::EExtension::KHR_LightsPunctual), GLTF::ToString(GLTF::EExtension::KHR_Lights) };
+		static const TArray<FString> MaterialsExtensions = { GLTF::ToString(GLTF::EExtension::KHR_MaterialsVariants) };
 	
 		TSharedPtr<FJsonObject> GetLightExtension(const TSharedPtr<FJsonObject>& Object)
 		{
@@ -61,7 +61,7 @@ namespace GLTF
 
 		if (TSharedPtr<FJsonObject> LightsObj = GetLightExtension(ExtensionsObj))
 		{
-			Asset->ExtensionsUsed.Add(EExtension::KHR_LightsPunctual);
+			Asset->ProcessedExtensions.Add(EExtension::KHR_LightsPunctual);
 
 			uint32 LightCount = ArraySize(*LightsObj, TEXT("lights"));
 			if (LightCount > 0)
@@ -76,11 +76,11 @@ namespace GLTF
 
 		// variants
 
-		if (ExtensionsObj->HasTypedField<EJson::Object>(TEXT("KHR_materials_variants")))
+		if (ExtensionsObj->HasTypedField<EJson::Object>(GLTF::ToString(GLTF::EExtension::KHR_MaterialsVariants)))
 		{
-			Asset->ExtensionsUsed.Add(EExtension::KHR_MaterialsVariants);
+			Asset->ProcessedExtensions.Add(EExtension::KHR_MaterialsVariants);
 
-			TSharedPtr<FJsonObject> VariantsObj = ExtensionsObj->GetObjectField(TEXT("KHR_materials_variants"));
+			TSharedPtr<FJsonObject> VariantsObj = ExtensionsObj->GetObjectField(GLTF::ToString(GLTF::EExtension::KHR_MaterialsVariants));
 			uint32 VariantsCount = ArraySize(*VariantsObj, TEXT("variants"));
 			if (VariantsCount > 0)
 			{
@@ -108,33 +108,36 @@ namespace GLTF
 			return;
 		}
 
-		enum
-		{
-			KHR_materials_pbrSpecularGlossiness = 0,
-			KHR_materials_unlit,
-			KHR_materials_clearcoat,
-			KHR_materials_transmission,
-			KHR_materials_sheen,
-			KHR_materials_ior,
-			KHR_materials_specular,
-			MSFT_packing_occlusionRoughnessMetallic,
-			MSFT_packing_normalRoughnessMetallic,
+		static const TArray<EExtension> Extensions = { 
+			EExtension::KHR_MaterialsPbrSpecularGlossiness,
+			EExtension::KHR_MaterialsUnlit,
+			EExtension::KHR_MaterialsClearCoat,
+			EExtension::KHR_MaterialsTransmission,
+			EExtension::KHR_MaterialsSheen,
+			EExtension::KHR_MaterialsIOR,
+			EExtension::KHR_MaterialsSpecular,
+			EExtension::MSFT_PackingOcclusionRoughnessMetallic,
+			EExtension::MSFT_PackingNormalRoughnessMetallic 
 		};
-		static const TArray<FString> Extensions = {TEXT("KHR_materials_pbrSpecularGlossiness"), TEXT("KHR_materials_unlit"), TEXT("KHR_materials_clearcoat"),
-												   TEXT("KHR_materials_transmission"), TEXT("KHR_materials_sheen"), TEXT("KHR_materials_ior"), TEXT("KHR_materials_specular"),
-		                                           TEXT("MSFT_packing_occlusionRoughnessMetallic"), TEXT("MSFT_packing_normalRoughnessMetallic")};
+		TArray<FString> ExtensionsStringified;
+		for (size_t ExtensionIndex = 0; ExtensionIndex < Extensions.Num(); ExtensionIndex++)
+		{
+			ExtensionsStringified.Add(GLTF::ToString(Extensions[ExtensionIndex]));
+		}
 
 		const FJsonObject& ExtensionsObj = *Object.GetObjectField(TEXT("extensions"));
 		for (int32 Index = 0; Index < Extensions.Num(); ++Index)
 		{
-			const FString ExtensionName = Extensions[Index];
+			const FString ExtensionName = ExtensionsStringified[Index];
 			if (!ExtensionsObj.HasTypedField<EJson::Object>(ExtensionName))
 				continue;
 
 			const FJsonObject& ExtObj = *ExtensionsObj.GetObjectField(ExtensionName);
-			switch (Index)
+
+			const EExtension Extension = Extensions[Index];
+			switch (Extension)
 			{
-				case KHR_materials_pbrSpecularGlossiness:
+				case EExtension::KHR_MaterialsPbrSpecularGlossiness:
 				{
 					const FJsonObject& PBR = ExtObj;
 					GLTF::SetTextureMap(PBR, TEXT("diffuseTexture"), nullptr, Asset->Textures, Material.BaseColor, Messages);
@@ -146,16 +149,16 @@ namespace GLTF
 
 					Material.ShadingModel = FMaterial::EShadingModel::SpecularGlossiness;
 
-					Asset->ExtensionsUsed.Add(EExtension::KHR_MaterialsPbrSpecularGlossiness);
+					Asset->ProcessedExtensions.Add(EExtension::KHR_MaterialsPbrSpecularGlossiness);
 				}
 				break;
-				case KHR_materials_unlit:
+				case EExtension::KHR_MaterialsUnlit:
 				{
 					Material.bIsUnlitShadingModel = true;
-					Asset->ExtensionsUsed.Add(EExtension::KHR_MaterialsUnlit);
+					Asset->ProcessedExtensions.Add(EExtension::KHR_MaterialsUnlit);
 				}
 				break;
-				case KHR_materials_clearcoat:
+				case EExtension::KHR_MaterialsClearCoat:
 				{
 					const FJsonObject& ClearCoat = ExtObj;
 
@@ -169,10 +172,10 @@ namespace GLTF
 
 					Material.ClearCoat.NormalMapUVScale = GLTF::SetTextureMap(ClearCoat, TEXT("clearcoatNormalTexture"), TEXT("scale"), Asset->Textures, Material.ClearCoat.NormalMap, Messages);
 
-					Asset->ExtensionsUsed.Add(EExtension::KHR_MaterialsClearCoat);
+					Asset->ProcessedExtensions.Add(EExtension::KHR_MaterialsClearCoat);
 				}
 				break;
-				case KHR_materials_transmission:
+				case EExtension::KHR_MaterialsTransmission:
 				{
 					if (Material.AlphaMode == FMaterial::EAlphaMode::Blend)
 					{
@@ -185,10 +188,10 @@ namespace GLTF
 					Material.Transmission.TransmissionFactor = GetScalar(Transm, TEXT("transmissionFactor"), 0.0f);
 					GLTF::SetTextureMap(Transm, TEXT("transmissionTexture"), nullptr, Asset->Textures, Material.Transmission.TransmissionMap, Messages);
 
-					Asset->ExtensionsUsed.Add(EExtension::KHR_MaterialsTransmission);
+					Asset->ProcessedExtensions.Add(EExtension::KHR_MaterialsTransmission);
 				}
 				break;
-				case KHR_materials_sheen:
+				case EExtension::KHR_MaterialsSheen:
 				{
 					const FJsonObject& Sheen = ExtObj;
 
@@ -200,20 +203,20 @@ namespace GLTF
 					Material.Sheen.SheenRoughnessFactor = GetScalar(Sheen, TEXT("sheenRoughnessFactor"));
 					GLTF::SetTextureMap(Sheen, TEXT("sheenRoughnessTexture"), nullptr, Asset->Textures, Material.Sheen.SheenRoughnessMap, Messages);
 
-					Asset->ExtensionsUsed.Add(EExtension::KHR_MaterialsSheen);
+					Asset->ProcessedExtensions.Add(EExtension::KHR_MaterialsSheen);
 				}
 				break;
-				case KHR_materials_ior:
+				case EExtension::KHR_MaterialsIOR:
 				{
 					const FJsonObject& IOR = ExtObj;
 
 					Material.bHasIOR = true;
 					Material.IOR = GetScalar(IOR, TEXT("ior"), 1.0f);
 
-					Asset->ExtensionsUsed.Add(EExtension::KHR_MaterialsIOR);
+					Asset->ProcessedExtensions.Add(EExtension::KHR_MaterialsIOR);
 				}
 				break;
-				case KHR_materials_specular:
+				case EExtension::KHR_MaterialsSpecular:
 				{
 					const FJsonObject& Specular = ExtObj;
 
@@ -223,10 +226,10 @@ namespace GLTF
 					GLTF::SetTextureMap(Specular, TEXT("specularTexture"), nullptr, Asset->Textures, Material.Specular.SpecularMap, Messages);
 					GLTF::SetTextureMap(Specular, TEXT("specularColorTexture"), nullptr, Asset->Textures, Material.Specular.SpecularColorMap, Messages);
 
-					Asset->ExtensionsUsed.Add(EExtension::KHR_MaterialsSpecular);
+					Asset->ProcessedExtensions.Add(EExtension::KHR_MaterialsSpecular);
 				}
 				break;
-				case MSFT_packing_occlusionRoughnessMetallic:
+				case EExtension::MSFT_PackingOcclusionRoughnessMetallic:
 				{
 					const FJsonObject& Packing = ExtObj;
 					if (GLTF::SetTextureMap(Packing, TEXT("occlusionRoughnessMetallicTexture"), nullptr, Asset->Textures, Material.Packing.Map, Messages))
@@ -244,27 +247,30 @@ namespace GLTF
 					}
 
 					if (Material.Packing.Flags != (int)GLTF::FMaterial::EPackingFlags::None)
-						Asset->ExtensionsUsed.Add(EExtension::MSFT_PackingOcclusionRoughnessMetallic);
+						Asset->ProcessedExtensions.Add(EExtension::MSFT_PackingOcclusionRoughnessMetallic);
 				}
 				break;
-				case MSFT_packing_normalRoughnessMetallic:
+				case EExtension::MSFT_PackingNormalRoughnessMetallic:
 				{
 					const FJsonObject& Packing = ExtObj;
 					if (GLTF::SetTextureMap(Packing, TEXT("normalRoughnessMetallicTexture"), nullptr, Asset->Textures, Material.Packing.Map, Messages))
 					{
 						Material.Packing.NormalMap = Material.Packing.Map;
 						Material.Packing.Flags     = (int)GLTF::FMaterial::EPackingFlags::NormalRoughnessMetallic;
-						Asset->ExtensionsUsed.Add(EExtension::MSFT_PackingNormalRoughnessMetallic);
+						Asset->ProcessedExtensions.Add(EExtension::MSFT_PackingNormalRoughnessMetallic);
 					}
 				}
 				break;
 				default:
-					check(false);
+					if (!ensure(false))
+					{
+						Messages.Emplace(EMessageSeverity::Warning, FString::Printf(TEXT("Material.Extension not supported: %s"), *ToString(Extension)));
+					}
 					break;
 			}
 		}
 
-		CheckExtensions(Object, Extensions);
+		CheckExtensions(Object, ExtensionsStringified);
 	}
 
 	void FExtensionsHandler::SetupBufferExtensions(const FJsonObject& Object, FBuffer& Buffer) const
@@ -292,23 +298,28 @@ namespace GLTF
 			return;
 		}
 
-		enum
-		{
-			KHR_materials_variants = 0
+		static const TArray<EExtension> Extensions = {
+			EExtension::KHR_MaterialsVariants
 		};
-		static const TArray<FString> Extensions = { TEXT("KHR_materials_variants") };
+		TArray<FString> ExtensionsStringified;
+		for (size_t ExtensionIndex = 0; ExtensionIndex < Extensions.Num(); ExtensionIndex++)
+		{
+			ExtensionsStringified.Add(GLTF::ToString(Extensions[ExtensionIndex]));
+		}
 
 		const FJsonObject& ExtensionsObj = *Object.GetObjectField(TEXT("extensions"));
 		for (int32 Index = 0; Index < Extensions.Num(); ++Index)
 		{
-			const FString ExtensionName = Extensions[Index];
+			const FString ExtensionName = ExtensionsStringified[Index];
 			if (!ExtensionsObj.HasTypedField<EJson::Object>(ExtensionName))
 				continue;
 
 			const FJsonObject& ExtObj = *ExtensionsObj.GetObjectField(ExtensionName);
-			switch (Index)
+
+			const EExtension Extension = Extensions[Index];
+			switch (Extension)
 			{
-				case KHR_materials_variants:
+				case EExtension::KHR_MaterialsVariants:
 				{
 					const TArray<TSharedPtr<FJsonValue>>& Mappings = ExtObj.GetArrayField(TEXT("mappings"));
 
@@ -328,16 +339,19 @@ namespace GLTF
 						}
 					}
 
-					Asset->ExtensionsUsed.Add(EExtension::KHR_MaterialsVariants);
+					Asset->ProcessedExtensions.Add(EExtension::KHR_MaterialsVariants);
 				}
 				break;
 				default:
-					check(false);
+					if (!ensure(false))
+					{
+						Messages.Emplace(EMessageSeverity::Warning, FString::Printf(TEXT("Primitive.Extension not supported: %s"), *ToString(Extension)));
+					}
 					break;
 			}
 		}
 
-		CheckExtensions(Object, Extensions);
+		CheckExtensions(Object, ExtensionsStringified);
 	}
 
 	void FExtensionsHandler::SetupMeshExtensions(const FJsonObject& Object, FMesh& Mesh) const
diff --git a/Engine/Plugins/Editor/GLTFImporter/Source/GLTFCore/Private/GLTF/GLTFAccessor.cpp b/Engine/Plugins/Editor/GLTFImporter/Source/GLTFCore/Private/GLTF/GLTFAccessor.cpp
--- a/Engine/Plugins/Editor/GLTFImporter/Source/GLTFCore/Private/GLTF/GLTFAccessor.cpp
+++ b/Engine/Plugins/Editor/GLTFImporter/Source/GLTFCore/Private/GLTF/GLTFAccessor.cpp
@@ -48,7 +48,7 @@ namespace GLTF
 				}
 			}
 			else
-				check(false);
+				ensure(false);
 
 			return Res;
 		}
@@ -113,7 +113,7 @@ namespace GLTF
 				}
 			}
 			else
-				check(false);
+				ensure(false);
 		}
 
 		// Copy data items that don't need conversion/expansion(i.e. Vec3 to Vec3, uint8 to uint8(not uint16)
@@ -309,7 +309,7 @@ namespace GLTF
 			}
 		}
 
-		check(false);
+		ensure(false);
 		return 0;
 	}
 
@@ -341,7 +341,7 @@ namespace GLTF
 				}
 			}
 		}
-		check(false);
+		ensure(false);
 	}
 
 	float FValidAccessor::GetFloat(uint32 Index) const
@@ -363,7 +363,7 @@ namespace GLTF
 			}
 		}
 
-		check(false);
+		ensure(false);
 		return 0.f;
 	}
 
@@ -393,7 +393,7 @@ namespace GLTF
 			}
 		}
 
-		check(false);
+		ensure(false);
 		return FVector2D::ZeroVector;
 	}
 
@@ -423,7 +423,7 @@ namespace GLTF
 			}
 		}
 
-		check(false);
+		ensure(false);
 		return FVector::ZeroVector;
 	}
 
@@ -453,7 +453,7 @@ namespace GLTF
 			}
 		}
 
-		check(false);
+		ensure(false);
 		return FVector4();
 	}
 
@@ -473,7 +473,7 @@ namespace GLTF
 			}
 		}
 
-		check(false);
+		ensure(false);
 		return FMatrix();
 	}
 
@@ -498,7 +498,7 @@ namespace GLTF
 			}
 		}
 
-		check(false);
+		ensure(false);
 	}
 
 	void FValidAccessor::GetFloatArray(float* Buffer) const
@@ -516,7 +516,7 @@ namespace GLTF
 			}
 		}
 
-		check(false);
+		ensure(false);
 	}
 
 	void FValidAccessor::GetVec2Array(FVector2f* Buffer) const
@@ -558,7 +558,7 @@ namespace GLTF
 			return;
 		}
 
-		check(false);
+		ensure(false);
 	}
 
 	inline const uint8* FValidAccessor::DataAt(uint32 Index) const
diff --git a/Engine/Plugins/Editor/GLTFImporter/Source/GLTFCore/Private/GLTF/GLTFAsset.cpp b/Engine/Plugins/Editor/GLTFImporter/Source/GLTFCore/Private/GLTF/GLTFAsset.cpp
--- a/Engine/Plugins/Editor/GLTFImporter/Source/GLTFCore/Private/GLTF/GLTFAsset.cpp
+++ b/Engine/Plugins/Editor/GLTFImporter/Source/GLTFCore/Private/GLTF/GLTFAsset.cpp
@@ -84,8 +84,8 @@ namespace GLTF
 		Textures.Empty();
 		Materials.Empty();
 
-		ExtensionsUsed.Empty((int)EExtension::Count);
-		RequiredExtensions.Empty();
+		ProcessedExtensions.Empty((int)EExtension::Count);
+		ExtensionsRequired.Empty();
 		Metadata.GeneratorName.Empty();
 		Metadata.Extras.Empty();
 
@@ -97,7 +97,7 @@ namespace GLTF
 
 	void FAsset::GenerateNames(const FString& Prefix)
 	{
-		check(!Prefix.IsEmpty());
+		ensure(!Prefix.IsEmpty());
 
 		{
 			const FString ScenePrefix = Prefix + TEXT("_scene_");
@@ -319,7 +319,7 @@ namespace GLTF
 			if (!VisitCount.Contains(NodeIndex))
 				NodeIndices.Add(NodeIndex);
 			else
-				check(VisitCount[NodeIndex] == 1);
+				ensure(VisitCount[NodeIndex] == 1);
 		}
 	}
 
@@ -367,35 +367,39 @@ namespace GLTF
 		switch (Extension)
 		{
 			case GLTF::EExtension::KHR_MaterialsPbrSpecularGlossiness:
-				return TEXT("KHR_Materials_PbrSpecularGlossiness");
+				return TEXT("KHR_materials_pbrSpecularGlossiness");
 			case GLTF::EExtension::KHR_MaterialsUnlit:
-				return TEXT("KHR_Materials_Unlit");
+				return TEXT("KHR_materials_unlit");
 			case GLTF::EExtension::KHR_MaterialsClearCoat:
-				return TEXT("KHR_MaterialsClearCoat");
-			case GLTF::EExtension::KHR_MaterialsSheen:
-				return TEXT("KHR_MaterialsSheen");
+				return TEXT("KHR_materials_clearcoat");
 			case GLTF::EExtension::KHR_MaterialsTransmission:
-				return TEXT("KHR_MaterialsTransmission");
-			case GLTF::EExtension::KHR_MaterialsSpecular:
-				return TEXT("KHR_MaterialsSpecular");
+				return TEXT("KHR_materials_transmission");
+			case GLTF::EExtension::KHR_MaterialsSheen:
+				return TEXT("KHR_materials_sheen");
+			case GLTF::EExtension::KHR_MaterialsVariants:
+				return TEXT("KHR_materials_variants");
 			case GLTF::EExtension::KHR_MaterialsIOR:
-				return TEXT("KHR_MaterialsIOR");
+				return TEXT("KHR_materials_ior");
+			case GLTF::EExtension::KHR_MaterialsSpecular:
+				return TEXT("KHR_materials_specular");
 			case GLTF::EExtension::KHR_TextureTransform:
-				return TEXT("KHR_Texture_Transform");
+				return TEXT("KHR_texture_transform");
 			case GLTF::EExtension::KHR_DracoMeshCompression:
-				return TEXT("KHR_DracoMeshCompression");
+				return TEXT("KHR_draco_mesh_compression");
 			case GLTF::EExtension::KHR_LightsPunctual:
-				return TEXT("KHR_LightsPunctual");
+				return TEXT("KHR_lights_punctual");
+			case GLTF::EExtension::KHR_Lights:
+				return TEXT("KHR_lights");
 			case GLTF::EExtension::KHR_Blend:
-				return TEXT("KHR_Blend");
+				return TEXT("KHR_blend");
 			case GLTF::EExtension::MSFT_TextureDDS:
-				return TEXT("MSFT_Texture_DDS");
+				return TEXT("MSFT_texture_dds");
 			case GLTF::EExtension::MSFT_PackingNormalRoughnessMetallic:
-				return TEXT("MSFT_Packing_NormalRoughnessMetallic");
+				return TEXT("MSFT_packing_normalRoughnessMetallic");
 			case GLTF::EExtension::MSFT_PackingOcclusionRoughnessMetallic:
-				return TEXT("MSFT_Packing_OcclusionRoughnessMetallic");
+				return TEXT("MSFT_packing_occlusionRoughnessMetallic");
 			case GLTF::EExtension::Count:
-				check(false);
+				ensure(false);
 			default:
 				return TEXT("UnknwonExtension");
 		}
diff --git a/Engine/Plugins/Editor/GLTFImporter/Source/GLTFCore/Private/GLTF/GLTFBinaryReader.cpp b/Engine/Plugins/Editor/GLTFImporter/Source/GLTFCore/Private/GLTF/GLTFBinaryReader.cpp
--- a/Engine/Plugins/Editor/GLTFImporter/Source/GLTFCore/Private/GLTF/GLTFBinaryReader.cpp
+++ b/Engine/Plugins/Editor/GLTFImporter/Source/GLTFCore/Private/GLTF/GLTFBinaryReader.cpp
@@ -97,7 +97,10 @@ namespace GLTF
 
 	bool FBinaryFileReader::ReadFile(FArchive& FileReader)
 	{
-		check(BinChunk);
+		if (!ensure(BinChunk))
+		{
+			return false;
+		}
 
 		Messages.Empty();
 
diff --git a/Engine/Plugins/Editor/GLTFImporter/Source/GLTFCore/Private/GLTF/GLTFReader.cpp b/Engine/Plugins/Editor/GLTFImporter/Source/GLTFCore/Private/GLTF/GLTFReader.cpp
--- a/Engine/Plugins/Editor/GLTFImporter/Source/GLTFCore/Private/GLTF/GLTFReader.cpp
+++ b/Engine/Plugins/Editor/GLTFImporter/Source/GLTFCore/Private/GLTF/GLTFReader.cpp
@@ -236,6 +236,11 @@ namespace GLTF
 	void FFileReader::SetupPrimitive(const FJsonObject& Object, FMesh& Mesh) const
 	{
 		const FPrimitive::EMode       Mode = PrimitiveModeFromNumber(GetUnsignedInt(Object, TEXT("mode"), (uint32)FPrimitive::EMode::Triangles));
+		if (Mode == FPrimitive::EMode::Unknown)
+		{
+			return;
+		}
+
 		const int32                   MaterialIndex = GetIndex(Object, TEXT("material"));
 		const TArray<FValidAccessor>& A             = Asset->Accessors;
 
@@ -929,8 +934,9 @@ namespace GLTF
 			OutAsset.Samplers.Empty(SamplerCount);
 			OutAsset.Textures.Empty(TextureCount);
 			OutAsset.Materials.Empty(MaterialCount);
-			OutAsset.ExtensionsUsed.Empty((int)EExtension::Count);
-			OutAsset.RequiredExtensions.Empty();
+			OutAsset.ProcessedExtensions.Empty((int)EExtension::Count);
+			OutAsset.ExtensionsUsed.Empty();
+			OutAsset.ExtensionsRequired.Empty();
 		}
 
 		// allocate asset mapped data for images and buffers
@@ -966,11 +972,23 @@ namespace GLTF
 		const TArray<TSharedPtr<FJsonValue>>* ExtensionsRequired;
 		if (JsonRoot->TryGetArrayField(TEXT("extensionsRequired"), ExtensionsRequired))
 		{
-			check(ExtensionsRequired);
-			OutAsset.RequiredExtensions.Reserve(OutAsset.RequiredExtensions.Num() + ExtensionsRequired->Num());
-			for (const TSharedPtr<FJsonValue>& Extension : *ExtensionsRequired)
+			if (ensure(ExtensionsRequired))
+			{
+				OutAsset.ExtensionsRequired.Reserve(OutAsset.ExtensionsRequired.Num() + ExtensionsRequired->Num());
+				for (const TSharedPtr<FJsonValue>& Extension : *ExtensionsRequired)
+				{
+					OutAsset.ExtensionsRequired.Add(Extension->AsString());
+				}
+			}
+		}
+
+		const TArray<TSharedPtr<FJsonValue>>* ExtensionsUsed;
+		if (JsonRoot->TryGetArrayField(TEXT("extensionsUsed"), ExtensionsUsed))
+		{
+			OutAsset.ExtensionsUsed.Reserve(ExtensionsUsed->Num());
+			for (const TSharedPtr<FJsonValue>& Extension : *ExtensionsUsed)
 			{
-				OutAsset.RequiredExtensions.Add(Extension->AsString());
+				OutAsset.ExtensionsUsed.Add(Extension->AsString());
 			}
 		}
 
diff --git a/Engine/Plugins/Editor/GLTFImporter/Source/GLTFCore/Private/GLTF/MaterialUtilities.h b/Engine/Plugins/Editor/GLTFImporter/Source/GLTFCore/Private/GLTF/MaterialUtilities.h
--- a/Engine/Plugins/Editor/GLTFImporter/Source/GLTFCore/Private/GLTF/MaterialUtilities.h
+++ b/Engine/Plugins/Editor/GLTFImporter/Source/GLTFCore/Private/GLTF/MaterialUtilities.h
@@ -34,18 +34,20 @@ namespace GLTF
 				// Handle textureInfo extensions
 				if (TexObj.HasTypedField<EJson::Object>(TEXT("extensions")))
 				{
-					enum
-					{
-						KHR_texture_transform = 0
+					static const TArray<EExtension> SupportedTextureInfoExtensions = {
+						EExtension::KHR_TextureTransform
 					};
-
-					static const TArray<FString> SupportedTextureInfoExtensions = { TEXT("KHR_texture_transform") };
+					TArray<FString> SupportedTextureInfoExtensionsStringified;
+					for (size_t ExtensionIndex = 0; ExtensionIndex < SupportedTextureInfoExtensions.Num(); ExtensionIndex++)
+					{
+						SupportedTextureInfoExtensionsStringified.Add(GLTF::ToString(SupportedTextureInfoExtensions[ExtensionIndex]));
+					}
 
 					const FJsonObject& ExtensionsObj = *TexObj.GetObjectField(TEXT("extensions"));
 
 					for (int32 Index = 0; Index < SupportedTextureInfoExtensions.Num(); ++Index)
 					{
-						const FString ExtensionName = SupportedTextureInfoExtensions[Index];
+						const FString ExtensionName = SupportedTextureInfoExtensionsStringified[Index];
 						if (!ExtensionsObj.HasTypedField<EJson::Object>(ExtensionName))
 						{
 							OutMessages.Emplace(EMessageSeverity::Warning, FString::Printf(TEXT("Extension is not supported: %s"), *ExtensionName));
@@ -53,9 +55,11 @@ namespace GLTF
 						}
 
 						const FJsonObject& ExtObj = *ExtensionsObj.GetObjectField(ExtensionName);
-						switch (Index)
+
+						const EExtension Extension = SupportedTextureInfoExtensions[Index];
+						switch (Extension)
 						{
-							case KHR_texture_transform:
+							case EExtension::KHR_TextureTransform:
 							{
 								OutMap.bHasTextureTransform = true;
 								OutMap.TextureTransform.Offset[0] = 0.0f;
diff --git a/Engine/Plugins/Editor/GLTFImporter/Source/GLTFCore/Private/GLTFMapFactory.cpp b/Engine/Plugins/Editor/GLTFImporter/Source/GLTFCore/Private/GLTFMapFactory.cpp
--- a/Engine/Plugins/Editor/GLTFImporter/Source/GLTFCore/Private/GLTFMapFactory.cpp
+++ b/Engine/Plugins/Editor/GLTFImporter/Source/GLTFCore/Private/GLTFMapFactory.cpp
@@ -127,7 +127,10 @@ namespace GLTF
 
 	void FPBRMapFactory::CreateNormalMap(const GLTF::FTexture& Map, int CoordinateIndex, float NormalScale, const GLTF::FTextureTransform* TextureTransform)
 	{
-		check(CurrentMaterialElement);
+		if (!ensure(CurrentMaterialElement))
+		{
+			return;
+		}
 
 		FMaterialExpressionTexture* TexExpression = CreateTextureMap(Map, CoordinateIndex, TEXT("Normal Map"), ETextureMode::Normal, TextureTransform);
 		if (!TexExpression)
@@ -317,8 +320,10 @@ namespace GLTF
 	                                               const TCHAR* ValueName, ETextureMode TextureMode, FMaterialExpressionInput& MaterialInput, 
 												   const GLTF::FTextureTransform* TextureTransform)
 	{
-		check(MapName);
-		check(CurrentMaterialElement);
+		if (!ensure(MapName) || !ensure(CurrentMaterialElement))
+		{
+			return nullptr;
+		}
 
 		ValueExpressionClass* ValueExpression = CurrentMaterialElement->AddMaterialExpression<ValueExpressionClass>();
 		if (ValueName)
diff --git a/Engine/Plugins/Editor/GLTFImporter/Source/GLTFCore/Private/GLTFMaterialExpressions.cpp b/Engine/Plugins/Editor/GLTFImporter/Source/GLTFCore/Private/GLTFMaterialExpressions.cpp
--- a/Engine/Plugins/Editor/GLTFImporter/Source/GLTFCore/Private/GLTFMaterialExpressions.cpp
+++ b/Engine/Plugins/Editor/GLTFImporter/Source/GLTFCore/Private/GLTFMaterialExpressions.cpp
@@ -69,28 +69,28 @@ namespace GLTF
 
 	void FMaterialExpressionParameter::SetName(const TCHAR* InName)
 	{
-		check(Type == EMaterialExpressionType::ConstantColor || Type == EMaterialExpressionType::ConstantScalar ||
+		ensure(Type == EMaterialExpressionType::ConstantColor || Type == EMaterialExpressionType::ConstantScalar ||
 		      Type == EMaterialExpressionType::Texture);
 		Name = InName;
 	}
 
 	const TCHAR* FMaterialExpressionParameter::GetName() const
 	{
-		check(Type == EMaterialExpressionType::ConstantColor || Type == EMaterialExpressionType::ConstantScalar ||
+		ensure(Type == EMaterialExpressionType::ConstantColor || Type == EMaterialExpressionType::ConstantScalar ||
 		      Type == EMaterialExpressionType::Texture);
 		return *Name;
 	}
 
 	void FMaterialExpressionParameter::SetGroupName(const TCHAR* InGroupName)
 	{
-		check(Type == EMaterialExpressionType::ConstantColor || Type == EMaterialExpressionType::ConstantScalar ||
+		ensure(Type == EMaterialExpressionType::ConstantColor || Type == EMaterialExpressionType::ConstantScalar ||
 		      Type == EMaterialExpressionType::Texture);
 		GroupName = InGroupName;
 	}
 
 	const TCHAR* FMaterialExpressionParameter::GetGroupName() const
 	{
-		check(Type == EMaterialExpressionType::ConstantColor || Type == EMaterialExpressionType::ConstantScalar ||
+		ensure(Type == EMaterialExpressionType::ConstantColor || Type == EMaterialExpressionType::ConstantScalar ||
 		      Type == EMaterialExpressionType::Texture);
 		return *GroupName;
 	}
@@ -99,7 +99,7 @@ namespace GLTF
 
 	FMaterialExpressionInput* FMaterialExpressionTexture::GetInput(int32 Index)
 	{
-		check(Index == 0);
+		ensure(Index == 0);
 		return &InputCoordinate;
 	}
 
@@ -185,7 +185,10 @@ namespace GLTF
 
 	FMaterialExpression* FMaterialElement::GetExpression(int32 Index)
 	{
-		check(Index < Expressions.Num());
+		if (!ensure(Index < Expressions.Num()))
+		{
+			return nullptr;
+		}
 		return Expressions[Index];
 	}
 
@@ -213,7 +216,7 @@ namespace GLTF
 				Expression = new FMaterialExpressionTextureCoordinate();
 				break;
 			default:
-				check(false);
+				ensure(false);
 		}
 		if (Expression)
 			Expressions.Add(Expression);
diff --git a/Engine/Plugins/Editor/GLTFImporter/Source/GLTFCore/Private/GLTFMaterialFactory.cpp b/Engine/Plugins/Editor/GLTFImporter/Source/GLTFCore/Private/GLTFMaterialFactory.cpp
--- a/Engine/Plugins/Editor/GLTFImporter/Source/GLTFCore/Private/GLTFMaterialFactory.cpp
+++ b/Engine/Plugins/Editor/GLTFImporter/Source/GLTFCore/Private/GLTFMaterialFactory.cpp
@@ -33,7 +33,10 @@ namespace GLTF
 		    : MaterialElementFactory(MaterialElementFactory)
 		    , TextureFactory(TextureFactory)
 		{
-			check(MaterialElementFactory);
+			if (!ensure(MaterialElementFactory))
+			{
+				Messages.Emplace(EMessageSeverity::Warning, FString::Printf(TEXT("MaterialElementFactory is unexpected nullptr.")));
+			}
 		}
 
 		const TArray<FMaterialElement*>& CreateMaterials(const GLTF::FAsset& Asset, UObject* ParentPackage, EObjectFlags Flags);
@@ -83,6 +86,12 @@ namespace GLTF
 			for (int32 Index = 0; Index < MaterialElement.GetExpressionsCount(); ++Index)
 			{
 				FMaterialExpression* Expression = MaterialElement.GetExpression(Index);
+
+				if (!Expression)
+				{
+					continue;
+				}
+
 				if (Expression->GetType() != EMaterialExpressionType::ConstantColor &&
 				    Expression->GetType() != EMaterialExpressionType::ConstantScalar && Expression->GetType() != EMaterialExpressionType::Texture)
 					continue;
@@ -90,8 +99,10 @@ namespace GLTF
 				FMaterialExpressionParameter* ExpressionParameter = static_cast<FMaterialExpressionParameter*>(Expression);
 				if (ExpressionParameter->GetName() == Name)
 				{
-					Result = static_cast<ReturnClass*>(ExpressionParameter);
-					check(Expression->GetType() == (EMaterialExpressionType)ReturnClass::Type);
+					if (ensure(Expression->GetType() == (EMaterialExpressionType)ReturnClass::Type))
+					{
+						Result = static_cast<ReturnClass*>(ExpressionParameter);
+					}
 					break;
 				}
 			}
@@ -104,6 +115,12 @@ namespace GLTF
 	{
 		TextureFactory->CleanUp();
 		Materials.Empty();
+
+		if (!MaterialElementFactory)
+		{
+			return Materials;
+		}
+
 		Materials.Reserve(Asset.Materials.Num() + 1);
 
 		Messages.Empty();
@@ -113,7 +130,10 @@ namespace GLTF
 
 		for (const GLTF::FMaterial& GLTFMaterial : Asset.Materials)
 		{
-			check(!GLTFMaterial.Name.IsEmpty());
+			if (!ensure(!GLTFMaterial.Name.IsEmpty()))
+			{
+				continue;
+			}
 
 			FMaterialElement* MaterialElement = MaterialElementFactory->CreateMaterial(*GLTFMaterial.Name, ParentPackage, Flags);
 			MaterialElement->SetTwoSided(GLTFMaterial.bIsDoubleSided);
@@ -204,7 +224,10 @@ namespace GLTF
 				break;
 			}
 			default:
-				check(false);
+				if (!ensure(false))
+				{
+					Messages.Emplace(EMessageSeverity::Warning, FString::Printf(TEXT("AlphaMode Opaque not supported for Opacity")));
+				}
 				break;
 		}
 	}
diff --git a/Engine/Plugins/Editor/GLTFImporter/Source/GLTFCore/Private/GLTFMeshFactory.cpp b/Engine/Plugins/Editor/GLTFImporter/Source/GLTFCore/Private/GLTFMeshFactory.cpp
--- a/Engine/Plugins/Editor/GLTFImporter/Source/GLTFCore/Private/GLTFMeshFactory.cpp
+++ b/Engine/Plugins/Editor/GLTFImporter/Source/GLTFCore/Private/GLTFMeshFactory.cpp
@@ -119,7 +119,10 @@ namespace GLTF
 		template <typename T>
 		void ReIndex(const TArray<T>& Source, const TArray<uint32>& Indices, TArray<T>& Dst)
 		{
-			check(&Source != &Dst);
+			if (!ensure(&Source != &Dst))
+			{
+				return;
+			}
 
 			Dst.Reserve(Indices.Num());
 			for (uint32 Index : Indices)
@@ -132,8 +135,11 @@ namespace GLTF
 		{
 			Normals.Empty();
 
-			const uint32 N = Indices.Num();
-			check(N % 3 == 0);
+			uint32 N = Indices.Num();
+			if (!ensure(N % 3 == 0))
+			{
+				N -= N % 3;
+			}
 			Normals.AddUninitialized(N);
 
 			for (uint32 i = 0; i < N; i += 3)
diff --git a/Engine/Plugins/Editor/GLTFImporter/Source/GLTFCore/Public/GLTFAsset.h b/Engine/Plugins/Editor/GLTFImporter/Source/GLTFCore/Public/GLTFAsset.h
--- a/Engine/Plugins/Editor/GLTFImporter/Source/GLTFCore/Public/GLTFAsset.h
+++ b/Engine/Plugins/Editor/GLTFImporter/Source/GLTFCore/Public/GLTFAsset.h
@@ -35,6 +35,7 @@ namespace GLTF
 		KHR_TextureTransform,
 		KHR_DracoMeshCompression,
 		KHR_LightsPunctual,
+		KHR_Lights,
 		KHR_Blend,
 		MSFT_TextureDDS,
 		MSFT_PackingNormalRoughnessMetallic,
@@ -95,10 +96,13 @@ namespace GLTF
 
 		TArray<FString>   Variants;
 
-		TSet<EExtension> ExtensionsUsed;
-		TArray<FString>	 RequiredExtensions;
+		TArray<FString>	 ExtensionsUsed;      //marked in the gltf
+		TArray<FString>	 ExtensionsRequired;  //marked in the gltf
+
+		TSet<EExtension> ProcessedExtensions;
 		FMetadata        Metadata;
 
+
 		/**
 		 * Will clear the asset's buffers.
 		 *
diff --git a/Engine/Plugins/Editor/GLTFImporter/Source/GLTFCore/Public/GLTFMesh.h b/Engine/Plugins/Editor/GLTFImporter/Source/GLTFCore/Public/GLTFMesh.h
--- a/Engine/Plugins/Editor/GLTFImporter/Source/GLTFCore/Public/GLTFMesh.h
+++ b/Engine/Plugins/Editor/GLTFImporter/Source/GLTFCore/Public/GLTFMesh.h
@@ -90,7 +90,10 @@ namespace GLTF
 			Triangles = 4,
 			// will be supported prior to release
 			TriangleStrip = 5,
-			TriangleFan   = 6
+			TriangleFan   = 6,
+
+			//
+			Unknown = 7
 		};
 
 		const EMode Mode;
@@ -225,7 +228,7 @@ namespace GLTF
 			case 1:
 				return TexCoord1.GetVec2Array(Buffer);
 			default:
-				check(false);
+				ensure(false);
 				break;
 		}
 	}
@@ -269,7 +272,7 @@ namespace GLTF
 			{
 				Count += Prim.HasJointWeights();
 			}
-			check(Primitives.Num() == Count);
+			ensure(Primitives.Num() == Count);
 		}
 
 		return Result;
diff --git a/Engine/Plugins/Editor/GLTFImporter/Source/GLTFImporter/Private/GLTFImportFactory.cpp b/Engine/Plugins/Editor/GLTFImporter/Source/GLTFImporter/Private/GLTFImportFactory.cpp
--- a/Engine/Plugins/Editor/GLTFImporter/Source/GLTFImporter/Private/GLTFImportFactory.cpp
+++ b/Engine/Plugins/Editor/GLTFImporter/Source/GLTFImporter/Private/GLTFImportFactory.cpp
@@ -206,7 +206,10 @@ void UGLTFImportFactory::UpdateMeshes() const
 	const FGLTFImporterContext& Context   = GLTFImporterModule->GetImporterContext();
 	const TArray<UStaticMesh*>& Meshes    = Context.StaticMeshFactory.GetMeshes();
 	const TArray<UMaterial*>&   Materials = Context.Materials;
-	check(Materials.Num() == Context.Asset.Materials.Num());
+	if (!ensure(Materials.Num() == Context.Asset.Materials.Num()))
+	{
+		return;
+	}
 
 	int32 MeshIndex = 0;
 	for (UStaticMesh* StaticMesh : Meshes)
@@ -221,7 +224,10 @@ void UGLTFImportFactory::UpdateMeshes() const
 			if (Primitive.MaterialIndex != INDEX_NONE)
 			{
 				Material = static_cast<UMaterialInterface*>(Materials[Primitive.MaterialIndex]);
-				check(Material);
+				if (!ensure(Material))
+				{
+					continue;
+				}
 			}
 
 			StaticMesh->GetStaticMaterials()[PrimIndex].MaterialInterface = Material;
diff --git a/Engine/Plugins/Editor/GLTFImporter/Source/GLTFImporter/Private/GLTFImporterContext.cpp b/Engine/Plugins/Editor/GLTFImporter/Source/GLTFImporter/Private/GLTFImporterContext.cpp
--- a/Engine/Plugins/Editor/GLTFImporter/Source/GLTFImporter/Private/GLTFImporterContext.cpp
+++ b/Engine/Plugins/Editor/GLTFImporter/Source/GLTFImporter/Private/GLTFImporterContext.cpp
@@ -44,10 +44,13 @@ bool FGLTFImporterContext::OpenFile(const FString& FilePath)
 	{
 		return false;
 	}
-	check(Asset.ValidationCheck() == GLTF::FAsset::Valid);
+	if (!ensure(Asset.ValidationCheck() == GLTF::FAsset::Valid))
+	{
+		return false;
+	}
 
 	// check extensions supported
-	for (GLTF::EExtension Extension : Asset.ExtensionsUsed)
+	for (GLTF::EExtension Extension : Asset.ProcessedExtensions)
 	{
 		if (!GLTF::FAsset::SupportedExtensions.Contains(Extension))
 		{
diff --git a/Engine/Plugins/Editor/GLTFImporter/Source/GLTFImporter/Private/GLTFMaterialElement.cpp b/Engine/Plugins/Editor/GLTFImporter/Source/GLTFImporter/Private/GLTFMaterialElement.cpp
--- a/Engine/Plugins/Editor/GLTFImporter/Source/GLTFImporter/Private/GLTFMaterialElement.cpp
+++ b/Engine/Plugins/Editor/GLTFImporter/Source/GLTFImporter/Private/GLTFMaterialElement.cpp
@@ -22,7 +22,7 @@ namespace GLTFImporterImpl
 	template <typename T>
 	T* NewMaterialExpression(UObject* Parent)
 	{
-		check(Parent != nullptr);
+		ensure(Parent != nullptr);
 
 		T* Expression                      = NewObject<T>(Parent);
 		Expression->MaterialExpressionGuid = FGuid::NewGuid();
@@ -65,7 +65,10 @@ namespace GLTFImporterImpl
 
 	UClass* FindClass(const TCHAR* ClassName)
 	{
-		check(ClassName);
+		if (!ensure(ClassName))
+		{
+			return nullptr;
+		}
 
 		if (UClass* Result = FindFirstObject<UClass>(ClassName, EFindFirstObjectOptions::None, ELogVerbosity::Warning, TEXT("GLTFImporter")))
 		{
@@ -82,7 +85,10 @@ namespace GLTFImporterImpl
 
 	UMaterialExpression* CreateTextureExpression(const GLTF::FMaterialExpression* Expression, UMaterial* UnrealMaterial)
 	{
-		check(Expression->GetType() == GLTF::EMaterialExpressionType::Texture);
+		if (!ensure(Expression->GetType() == GLTF::EMaterialExpressionType::Texture))
+		{
+			return nullptr;
+		}
 		const GLTF::FMaterialExpressionTexture& TextureExpression = *static_cast<const GLTF::FMaterialExpressionTexture*>(Expression);
 
 		UMaterialExpressionTextureSampleParameter2D* MaterialExpression =
@@ -101,7 +107,10 @@ namespace GLTFImporterImpl
 
 	UMaterialExpression* CreateTextureCoordinateExpression(const GLTF::FMaterialExpression* Expression, UMaterial* UnrealMaterial)
 	{
-		check(Expression->GetType() == GLTF::EMaterialExpressionType::TextureCoordinate);
+		if (!ensure(Expression->GetType() == GLTF::EMaterialExpressionType::TextureCoordinate))
+		{
+			return nullptr;
+		}
 		const GLTF::FMaterialExpressionTextureCoordinate& TextureCoordinateExpression =
 		    *static_cast<const GLTF::FMaterialExpressionTextureCoordinate*>(Expression);
 
@@ -113,7 +122,10 @@ namespace GLTFImporterImpl
 
 	UMaterialExpression* CreateGenericExpression(const GLTF::FMaterialExpression* Expression, UMaterial* UnrealMaterial)
 	{
-		check(Expression->GetType() == GLTF::EMaterialExpressionType::Generic);
+		if (!ensure(Expression->GetType() == GLTF::EMaterialExpressionType::Generic))
+		{
+			return nullptr;
+		}
 		const GLTF::FMaterialExpressionGeneric& GenericExpression = *static_cast<const GLTF::FMaterialExpressionGeneric*>(Expression);
 
 		UClass* ExpressionClass = FindClass(*(FString(TEXT("MaterialExpression")) + GenericExpression.GetExpressionName()));
@@ -140,7 +152,10 @@ namespace GLTFImporterImpl
 
 	UMaterialExpression* CreateFunctionCallExpression(const GLTF::FMaterialExpression* Expression, UMaterial* UnrealMaterial)
 	{
-		check(Expression->GetType() == GLTF::EMaterialExpressionType::FunctionCall);
+		if (!ensure(Expression->GetType() == GLTF::EMaterialExpressionType::FunctionCall))
+		{
+			return nullptr;
+		}
 		const GLTF::FMaterialExpressionFunctionCall& FunctionCall = *static_cast<const GLTF::FMaterialExpressionFunctionCall*>(Expression);
 
 		const FSoftObjectPath       MaterialFunctionObjectPath(FString(FunctionCall.GetFunctionPathName()));
@@ -155,7 +170,10 @@ namespace GLTFImporterImpl
 
 	UMaterialExpression* CreateScalarExpression(const GLTF::FMaterialExpression* Expression, UMaterial* UnrealMaterial)
 	{
-		check(Expression->GetType() == GLTF::EMaterialExpressionType::ConstantScalar);
+		if (!ensure(Expression->GetType() == GLTF::EMaterialExpressionType::ConstantScalar))
+		{
+			return nullptr;
+		}
 		const GLTF::FMaterialExpressionScalar& ScalarExpression = *static_cast<const GLTF::FMaterialExpressionScalar*>(Expression);
 
 		UMaterialExpression* Result = nullptr;
@@ -181,7 +199,10 @@ namespace GLTFImporterImpl
 
 	UMaterialExpression* CreateColorExpression(const GLTF::FMaterialExpression* Expression, UMaterial* UnrealMaterial)
 	{
-		check(Expression->GetType() == GLTF::EMaterialExpressionType::ConstantColor);
+		if (!ensure(Expression->GetType() == GLTF::EMaterialExpressionType::ConstantColor))
+		{
+			return nullptr;
+		}
 		const GLTF::FMaterialExpressionColor& ColorExpression = *static_cast<const GLTF::FMaterialExpressionColor*>(Expression);
 
 		UMaterialExpression* Result = nullptr;
@@ -210,41 +231,49 @@ FGLTFMaterialElement::FGLTFMaterialElement(UMaterial* Material)
     : GLTF::FMaterialElement(Material->GetName())
     , Material(Material)
 {
-	check(Material);
+	ensure(Material);
 }
 
 int FGLTFMaterialElement::GetBlendMode() const
 {
-	return Material->BlendMode;
+	if (Material) return Material->BlendMode;
+	
+	return 0;
 }
 
 void FGLTFMaterialElement::SetBlendMode(int InBlendMode)
 {
-	Material->BlendMode = static_cast<EBlendMode>(InBlendMode);
+	if (Material) Material->BlendMode = static_cast<EBlendMode>(InBlendMode);
 }
 
 bool FGLTFMaterialElement::GetTwoSided() const
 {
-	return Material->IsTwoSided();
+	if (Material) return Material->IsTwoSided();
+
+	return false;
 }
 
 void FGLTFMaterialElement::SetTwoSided(bool bTwoSided)
 {
-	Material->TwoSided = bTwoSided;
+	if (Material) Material->TwoSided = bTwoSided;
 }
 
 bool FGLTFMaterialElement::GetIsThinSurface() const
 {
-	return Material->IsThinSurface();
+	if (Material) return Material->IsThinSurface();
+
+	return false;
 }
 
 void FGLTFMaterialElement::SetIsThinSurface(bool bIsThinSurface)
 {
-	Material->bIsThinSurface = bIsThinSurface;
+	if (Material) Material->bIsThinSurface = bIsThinSurface;
 }
 
 void FGLTFMaterialElement::SetShadingModel(GLTF::EGLTFMaterialShadingModel InShadingModel)
 {
+	if (!Material) return;
+
 	EMaterialShadingModel MaterialShadingModel;
 
 	switch (InShadingModel)
@@ -266,12 +295,16 @@ void FGLTFMaterialElement::SetShadingModel(GLTF::EGLTFMaterialShadingModel InSha
 
 void FGLTFMaterialElement::SetTranslucencyLightingMode(int InLightingMode)
 {
-	Material->TranslucencyLightingMode = static_cast<ETranslucencyLightingMode>( InLightingMode );
+	if (Material) Material->TranslucencyLightingMode = static_cast<ETranslucencyLightingMode>( InLightingMode );
 }
 
 void FGLTFMaterialElement::Finalize()
 {
-	check(!bIsFinal);
+	if (!ensure(!bIsFinal) || !Material)
+	{
+		//already finalized.
+		return;
+	}
 
 	TArray<TStrongObjectPtr<UMaterialExpression> > MaterialExpressions;
 	CreateExpressions(MaterialExpressions);
@@ -341,7 +374,10 @@ void FGLTFMaterialElement::CreateExpressions(TArray<TStrongObjectPtr<UMaterialEx
 	{
 		using namespace GLTFImporterImpl;
 
-		check(Expression);
+		if (!ensure(Expression))
+		{
+			continue;
+		}
 
 		UMaterialExpression* MaterialExpression = nullptr;
 		switch (Expression->GetType())
@@ -368,7 +404,7 @@ void FGLTFMaterialElement::CreateExpressions(TArray<TStrongObjectPtr<UMaterialEx
 				break;
 		}
 
-		check(MaterialExpression);
+		ensure(MaterialExpression);
 
 		MaterialExpressions.Add(TStrongObjectPtr<UMaterialExpression>(MaterialExpression));
 
@@ -397,17 +433,16 @@ void FGLTFMaterialElement::ConnectExpression(const GLTF::FMaterialExpression*
                                              FExpressionInput&                                     ExpressionInput,      //
                                              int32                                                 OutputIndex)
 {
-	check(Expressions.Num() == MaterialExpressions.Num());
+	ensure(Expressions.Num() == MaterialExpressions.Num());
 
 	if (ExpressionPtr == nullptr)
 		return;
 
 	GLTF::FMaterialExpression& Expression      = *const_cast<GLTF::FMaterialExpression*>(ExpressionPtr);  // safe as we dont modify it
 	const int32                ExpressionIndex = Expressions.Find(&Expression);
-	check(ExpressionIndex != INDEX_NONE);
-	if (!MaterialExpressions.IsValidIndex(ExpressionIndex))
+	if (ExpressionIndex != INDEX_NONE  || !MaterialExpressions.IsValidIndex(ExpressionIndex))
 	{
-		check(false);
+		ensure(false);
 		return;
 	}
 
diff --git a/Engine/Plugins/Editor/GLTFImporter/Source/GLTFImporter/Private/GLTFStaticMeshFactory.cpp b/Engine/Plugins/Editor/GLTFImporter/Source/GLTFImporter/Private/GLTFStaticMeshFactory.cpp
--- a/Engine/Plugins/Editor/GLTFImporter/Source/GLTFImporter/Private/GLTFStaticMeshFactory.cpp
+++ b/Engine/Plugins/Editor/GLTFImporter/Source/GLTFImporter/Private/GLTFStaticMeshFactory.cpp
@@ -71,7 +71,10 @@ namespace GLTF
 
 	UStaticMesh* FStaticMeshFactoryImpl::CreateMesh(const FMesh& Mesh, UObject* ParentPackage, EObjectFlags Flags)
 	{
-		check(!Mesh.Name.IsEmpty());
+		if (!ensure(!Mesh.Name.IsEmpty()))
+		{
+			return nullptr;
+		}
 
 		const FString PackageName  = UPackageTools::SanitizePackageName(FPaths::Combine(ParentPackage->GetName(), Mesh.Name));
 		UPackage*     AssetPackage = CreatePackage(*PackageName);
diff --git a/Engine/Plugins/Editor/GLTFImporter/Source/GLTFImporter/Private/GLTFTextureFactory.cpp b/Engine/Plugins/Editor/GLTFImporter/Source/GLTFImporter/Private/GLTFTextureFactory.cpp
--- a/Engine/Plugins/Editor/GLTFImporter/Source/GLTFImporter/Private/GLTFTextureFactory.cpp
+++ b/Engine/Plugins/Editor/GLTFImporter/Source/GLTFImporter/Private/GLTFTextureFactory.cpp
@@ -74,7 +74,7 @@ namespace GLTFImporterImpl
 				return TEXT("jpeg");
 				break;
 			default:
-				check(false);
+				ensure(false);
 				return TEXT("unknown");
 		}
 	}
@@ -110,7 +110,7 @@ GLTF::ITextureElement* FGLTFTextureFactory::CreateTexture(const GLTF::FTexture&
 	UTexture2D* Texture = nullptr;
 	if (!FPaths::GetExtension(GltfTexture.Source.FilePath).IsEmpty())
 	{
-		check(GltfTexture.Source.Data == nullptr);
+		ensure(GltfTexture.Source.Data == nullptr);
 		bool bOperationCanceled = false;
 		Texture                 = static_cast<UTexture2D*>(Factory->FactoryCreateFile(UTexture2D::StaticClass(), AssetPackage, *TextureName, Flags,
                                                                       GltfTexture.Source.FilePath, nullptr, nullptr, bOperationCanceled));
@@ -121,7 +121,7 @@ GLTF::ITextureElement* FGLTFTextureFactory::CreateTexture(const GLTF::FTexture&
 	}
 	else if (GltfTexture.Source.Format != GLTF::FImage::EFormat::Unknown)
 	{
-		check(GltfTexture.Source.Data != nullptr);
+		ensure(GltfTexture.Source.Data != nullptr);
 
 		const uint8* PtrTextureData    = GltfTexture.Source.Data;
 		const uint8* PtrTextureDataEnd = PtrTextureData + GltfTexture.Source.DataByteLength;
diff --git a/Engine/Plugins/Enterprise/DatasmithGLTFImporter/Source/DatasmithGLTFTranslator/Private/DatasmithGLTFImporter.cpp b/Engine/Plugins/Enterprise/DatasmithGLTFImporter/Source/DatasmithGLTFTranslator/Private/DatasmithGLTFImporter.cpp
--- a/Engine/Plugins/Enterprise/DatasmithGLTFImporter/Source/DatasmithGLTFTranslator/Private/DatasmithGLTFImporter.cpp
+++ b/Engine/Plugins/Enterprise/DatasmithGLTFImporter/Source/DatasmithGLTFTranslator/Private/DatasmithGLTFImporter.cpp
@@ -96,7 +96,7 @@ bool FDatasmithGLTFImporter::OpenFile(const FString& InFileName)
 	check(GLTFAsset->ValidationCheck() == GLTF::FAsset::Valid);
 
 	// check extensions supported
-	for (GLTF::EExtension Extension : GLTFAsset->ExtensionsUsed)
+	for (GLTF::EExtension Extension : GLTFAsset->ProcessedExtensions)
 	{
 		if (!GLTF::FAsset::SupportedExtensions.Contains(Extension))
 		{
diff --git a/Engine/Plugins/Interchange/Runtime/Source/Import/Private/Gltf/InterchangeGltfTranslator.cpp b/Engine/Plugins/Interchange/Runtime/Source/Import/Private/Gltf/InterchangeGltfTranslator.cpp
--- a/Engine/Plugins/Interchange/Runtime/Source/Import/Private/Gltf/InterchangeGltfTranslator.cpp
+++ b/Engine/Plugins/Interchange/Runtime/Source/Import/Private/Gltf/InterchangeGltfTranslator.cpp
@@ -52,23 +52,24 @@ static FAutoConsoleVariableRef CCvarInterchangeEnableGLTFImport(
 	TEXT("Whether glTF support is enabled."),
 	ECVF_Default);
 
-static const TArray<FString> SupportedExtensions = {/* Lights */
-													TEXT("KHR_lights_punctual"),
-													TEXT("KHR_lights"),
-													/* Variants */
-													TEXT("KHR_materials_variants"), 
-													/* Materials */
-													TEXT("KHR_materials_unlit"),
-													TEXT("KHR_materials_ior"),
-													TEXT("KHR_materials_clearcoat"),
-													TEXT("KHR_materials_transmission"),
-													TEXT("KHR_materials_sheen"),
-													TEXT("KHR_materials_specular"),
-													TEXT("KHR_materials_pbrSpecularGlossiness"),
-													TEXT("MSFT_packing_occlusionRoughnessMetallic"),
-													TEXT("MSFT_packing_normalRoughnessMetallic"),
-													/* Textures */
-													TEXT("KHR_texture_transform")
+static const TArray<FString> ImporterSupportedExtensions = {
+	/* Lights */
+	GLTF::ToString(GLTF::EExtension::KHR_LightsPunctual),
+	GLTF::ToString(GLTF::EExtension::KHR_Lights),
+	/* Variants */
+	GLTF::ToString(GLTF::EExtension::KHR_MaterialsVariants),
+	/* Materials */
+	GLTF::ToString(GLTF::EExtension::KHR_MaterialsUnlit),
+	GLTF::ToString(GLTF::EExtension::KHR_MaterialsIOR),
+	GLTF::ToString(GLTF::EExtension::KHR_MaterialsClearCoat),
+	GLTF::ToString(GLTF::EExtension::KHR_MaterialsTransmission),
+	GLTF::ToString(GLTF::EExtension::KHR_MaterialsSheen),
+	GLTF::ToString(GLTF::EExtension::KHR_MaterialsSpecular),
+	GLTF::ToString(GLTF::EExtension::KHR_MaterialsPbrSpecularGlossiness),
+	GLTF::ToString(GLTF::EExtension::MSFT_PackingOcclusionRoughnessMetallic),
+	GLTF::ToString(GLTF::EExtension::MSFT_PackingNormalRoughnessMetallic),
+	/* Textures */
+	GLTF::ToString(GLTF::EExtension::KHR_TextureTransform)
 };
 
 namespace UE::Interchange::Gltf::Private
@@ -143,31 +144,44 @@ namespace UE::Interchange::Gltf::Private
 		NOTSUPPORTED_EXTENSION_FOUND
 	};
 	void SendAnalytics(const TranslationResult& TranslationResult,
-		const FString& NotSupportedExtensions = FString(),
-		const TSet<GLTF::EExtension>& ExtensionsUsed = TSet<GLTF::EExtension>(),
-		const TArray<FString>& RequiredExtensions = TArray<FString>(),
+		const TArray<FString>& ExtensionsUsed = TArray<FString>(),
+		const TArray<FString>& ExtensionsRequired = TArray<FString>(),
 		GLTF::FMetadata Metadata = GLTF::FMetadata(),
 		const FString& GLTFReaderLogMessage = "")
 	{
 		if (FEngineAnalytics::IsAvailable())
 		{
-			TArray<FString> ExtensionsUsedStringified;
-			for (GLTF::EExtension Extension : ExtensionsUsed)
-			{
-				ExtensionsUsedStringified.Add(FString(GLTF::ToString(Extension)));
-			}
-
 			TMap<FString, FString> MetadataExtras;
 			for (GLTF::FMetadata::FExtraData ExtraData : Metadata.Extras)
 			{
 				MetadataExtras.Add(ExtraData.Name, ExtraData.Value);
 			}
 
+			TSet<FString> AllExtensions;
+			AllExtensions.Append(ExtensionsUsed);
+			AllExtensions.Append(ExtensionsRequired);
+
+			TArray<FString> ExtensionsSupported;
+			TArray<FString> ExtensionsUnsupported;
+
+			for (const FString& Extension : AllExtensions)
+			{
+				if (ImporterSupportedExtensions.Find(Extension) == INDEX_NONE)
+				{
+					ExtensionsUnsupported.Add(Extension);
+				}
+				else
+				{
+					ExtensionsSupported.Add(Extension);
+				}
+			}
+
 			TArray<FAnalyticsEventAttribute> GLTFAnalytics;
-			if (NotSupportedExtensions.Len() > 0)			GLTFAnalytics.Add(FAnalyticsEventAttribute(TEXT("NotSupportedExtensions"), NotSupportedExtensions));
-			if (RequiredExtensions.Num() > 0)		GLTFAnalytics.Add(FAnalyticsEventAttribute(TEXT("RequiredExtensions"), RequiredExtensions));
-			if (ExtensionsUsedStringified.Num() > 0)		GLTFAnalytics.Add(FAnalyticsEventAttribute(TEXT("UsedExtensions"), ExtensionsUsedStringified));
-			if (Metadata.GeneratorName.Len() > 0)	GLTFAnalytics.Add(FAnalyticsEventAttribute(TEXT("MetaData.GeneratorName"), Metadata.GeneratorName));
+			if (ExtensionsUsed.Num() > 0)					GLTFAnalytics.Add(FAnalyticsEventAttribute(TEXT("ExtensionsUsed"), ExtensionsUsed));
+			if (ExtensionsRequired.Num() > 0)				GLTFAnalytics.Add(FAnalyticsEventAttribute(TEXT("ExtensionsRequired"), ExtensionsRequired));
+			if (ExtensionsSupported.Num() > 0)				GLTFAnalytics.Add(FAnalyticsEventAttribute(TEXT("ExtensionsSupported"), ExtensionsSupported));
+			if (ExtensionsUnsupported.Num() > 0)			GLTFAnalytics.Add(FAnalyticsEventAttribute(TEXT("ExtensionsUnsupported"), ExtensionsUnsupported));
+			if (Metadata.GeneratorName.Len() > 0)			GLTFAnalytics.Add(FAnalyticsEventAttribute(TEXT("MetaData.GeneratorName"), Metadata.GeneratorName));
 			if (MetadataExtras.Num() > 0)					GLTFAnalytics.Add(FAnalyticsEventAttribute(TEXT("MetaData.Extras"), MetadataExtras));
 			/*Version is always set at this point.*/		GLTFAnalytics.Add(FAnalyticsEventAttribute(TEXT("MetaData.Version"), Metadata.Version));
 
@@ -964,18 +978,14 @@ bool UInterchangeGltfTranslator::Translate( UInterchangeBaseNodeContainer& NodeC
 	const FString FileName = GltfAsset.Name;
 
 	//Required Extension Check:
-	FString NotSupportedExtensions;
-	if (GltfAsset.RequiredExtensions.Num() != 0)
+	TArray<FString> NotSupportedRequiredExtensions;
+	if (GltfAsset.ExtensionsRequired.Num() != 0)
 	{
-		for (const FString& RequiredExtension : GltfAsset.RequiredExtensions)
+		for (const FString& RequiredExtension : GltfAsset.ExtensionsRequired)
 		{
-			if (SupportedExtensions.Find(RequiredExtension) == INDEX_NONE)
+			if (ImporterSupportedExtensions.Find(RequiredExtension) == INDEX_NONE)
 			{
-				if (NotSupportedExtensions.Len() > 0)
-				{
-					NotSupportedExtensions += ", ";
-				}
-				NotSupportedExtensions += RequiredExtension;
+				NotSupportedRequiredExtensions.Add(RequiredExtension);
 			}
 		}
 	}
@@ -990,21 +1000,31 @@ bool UInterchangeGltfTranslator::Translate( UInterchangeBaseNodeContainer& NodeC
 			ErrorResult->SourceAssetName = FileName;
 			ErrorResult->Text = FText::Format(LOCTEXT("GLTF::FFileReader::ReadFile Failed.", "LogMessage: {0}"), FText::FromString(LogMessage.Value));
 
-			SendAnalytics(TranslationResult::GLTFREADER_FAILED, NotSupportedExtensions, GltfAsset.ExtensionsUsed, GltfAsset.RequiredExtensions, GltfAsset.Metadata, LogMessage.Value);
+			SendAnalytics(TranslationResult::GLTFREADER_FAILED, GltfAsset.ExtensionsUsed,GltfAsset.ExtensionsRequired, GltfAsset.Metadata, LogMessage.Value);
 			return false;
 		}
 	}
 
 	//In case of non supported extensions fail out:
-	if (NotSupportedExtensions.Len() > 0)
+	if (NotSupportedRequiredExtensions.Num() > 0)
 	{
+		FString NotSupportedRequiredExtensionsStringified;
+		for (const FString& NotSupportedExtension : NotSupportedRequiredExtensions)
+		{
+			if (NotSupportedRequiredExtensionsStringified.Len() > 0)
+			{
+				NotSupportedRequiredExtensionsStringified += ", ";
+			}
+			NotSupportedRequiredExtensionsStringified += NotSupportedExtension;
+		}
+
 		UInterchangeResultError_Generic* ErrorResult = AddMessage< UInterchangeResultError_Generic >();
 		ErrorResult->SourceAssetName = FileName;
 		ErrorResult->Text = FText::Format(
 			LOCTEXT("UnsupportedRequiredExtensions", "Not All Required Extensions are supported. (Unsupported extensions: {0})"),
-			FText::FromString(NotSupportedExtensions));
+			FText::FromString(NotSupportedRequiredExtensionsStringified));
 
-		SendAnalytics(TranslationResult::NOTSUPPORTED_EXTENSION_FOUND, NotSupportedExtensions, GltfAsset.ExtensionsUsed, GltfAsset.RequiredExtensions, GltfAsset.Metadata);
+		SendAnalytics(TranslationResult::NOTSUPPORTED_EXTENSION_FOUND, GltfAsset.ExtensionsUsed, GltfAsset.ExtensionsRequired, GltfAsset.Metadata);
 		return false;
 	}
 
@@ -1254,7 +1274,7 @@ bool UInterchangeGltfTranslator::Translate( UInterchangeBaseNodeContainer& NodeC
 		UE_LOG(LogInterchangeImport, Warning, TEXT("GLTF Mesh Import Warning. Gltf Mesh Usage expectation is not met."));
 	}
 
-	SendAnalytics(TranslationResult::SUCCESSFULL, NotSupportedExtensions, GltfAsset.ExtensionsUsed, GltfAsset.RequiredExtensions, GltfAsset.Metadata);;
+	SendAnalytics(TranslationResult::SUCCESSFULL, GltfAsset.ExtensionsUsed, GltfAsset.ExtensionsRequired, GltfAsset.Metadata);
 	return true;
 }
 
