Static analyzer warnings in GLTFMaterialElement.cpp fixes

diff --git a/Engine/Plugins/Editor/GLTFImporter/Source/GLTFImporter/Private/GLTFMaterialElement.cpp b/Engine/Plugins/Editor/GLTFImporter/Source/GLTFImporter/Private/GLTFMaterialElement.cpp
--- a/Engine/Plugins/Editor/GLTFImporter/Source/GLTFImporter/Private/GLTFMaterialElement.cpp
+++ b/Engine/Plugins/Editor/GLTFImporter/Source/GLTFImporter/Private/GLTFMaterialElement.cpp
@@ -22,7 +22,10 @@ namespace GLTFImporterImpl
 	template <typename T>
 	T* NewMaterialExpression(UObject* Parent)
 	{
-		ensure(Parent != nullptr);
+		if (!ensure(Parent))
+		{
+			return nullptr;
+		}
 
 		T* Expression                      = NewObject<T>(Parent);
 		Expression->MaterialExpressionGuid = FGuid::NewGuid();
@@ -44,6 +47,12 @@ namespace GLTFImporterImpl
 	T* NewMaterialExpressionParameter(UObject* Parent, const FString& Name)
 	{
 		T* Expression              = NewMaterialExpression<T>(Parent);
+
+		if (!Expression)
+		{
+			return nullptr;
+		}
+
 		Expression->ExpressionGUID = FGuid::NewGuid();
 		Expression->ParameterName  = *Name;
 		return Expression;
@@ -115,6 +124,11 @@ namespace GLTFImporterImpl
 		    *static_cast<const GLTF::FMaterialExpressionTextureCoordinate*>(Expression);
 
 		UMaterialExpressionTextureCoordinate* MaterialExpression = NewMaterialExpression<UMaterialExpressionTextureCoordinate>(UnrealMaterial);
+		if (!MaterialExpression)
+		{
+			return nullptr;
+		}
+
 		MaterialExpression->CoordinateIndex                      = TextureCoordinateExpression.GetCoordinateIndex();
 
 		return MaterialExpression;
@@ -162,6 +176,11 @@ namespace GLTFImporterImpl
 		UMaterialFunctionInterface* MaterialFunction = Cast<UMaterialFunctionInterface>(MaterialFunctionObjectPath.TryLoad());
 
 		UMaterialExpressionMaterialFunctionCall* MaterialExpression = NewMaterialExpression<UMaterialExpressionMaterialFunctionCall>(UnrealMaterial);
+		if (!MaterialExpression)
+		{
+			return nullptr;
+		}
+
 		MaterialExpression->SetMaterialFunction(MaterialFunction);
 		MaterialExpression->UpdateFromFunctionResource();
 
@@ -180,6 +199,11 @@ namespace GLTFImporterImpl
 		if (FCString::Strlen(ScalarExpression.GetName()) == 0)
 		{
 			UMaterialExpressionConstant* MaterialExpression = NewMaterialExpression<UMaterialExpressionConstant>(UnrealMaterial);
+			if (!MaterialExpression)
+			{
+				return nullptr;
+			}
+
 			MaterialExpression->R                           = ScalarExpression.GetScalar();
 
 			Result = MaterialExpression;
@@ -188,6 +212,11 @@ namespace GLTFImporterImpl
 		{
 			UMaterialExpressionScalarParameter* MaterialExpression =
 			    NewMaterialExpressionParameter<UMaterialExpressionScalarParameter>(UnrealMaterial, ScalarExpression.GetName());
+			if (!MaterialExpression)
+			{
+				return nullptr;
+			}
+
 			MaterialExpression->DefaultValue = ScalarExpression.GetScalar();
 			MaterialExpression->Group        = ScalarExpression.GetGroupName();
 
@@ -209,6 +238,11 @@ namespace GLTFImporterImpl
 		if (FCString::Strlen(ColorExpression.GetName()) == 0)
 		{
 			UMaterialExpressionConstant3Vector* MaterialExpression = NewMaterialExpression<UMaterialExpressionConstant3Vector>(UnrealMaterial);
+			if (!MaterialExpression)
+			{
+				return nullptr;
+			}
+
 			MaterialExpression->Constant                           = ColorExpression.GetColor();
 
 			Result = MaterialExpression;
@@ -217,6 +251,11 @@ namespace GLTFImporterImpl
 		{
 			UMaterialExpressionVectorParameter* MaterialExpression =
 			    NewMaterialExpressionParameter<UMaterialExpressionVectorParameter>(UnrealMaterial, ColorExpression.GetName());
+			if (!MaterialExpression)
+			{
+				return nullptr;
+			}
+
 			MaterialExpression->DefaultValue = ColorExpression.GetColor();
 			MaterialExpression->Group        = ColorExpression.GetGroupName();
 
@@ -404,7 +443,10 @@ void FGLTFMaterialElement::CreateExpressions(TArray<TStrongObjectPtr<UMaterialEx
 				break;
 		}
 
-		ensure(MaterialExpression);
+		if (!ensure(MaterialExpression))
+		{
+			continue;
+		}
 
 		MaterialExpressions.Add(TStrongObjectPtr<UMaterialExpression>(MaterialExpression));
 
