[Interchange] glTF imported vertex colors look washed out.
[GLTFCore] Countering sRGB conversion on VertexColor before uploading them to MeshDescription.

diff --git a/Engine/Plugins/Editor/GLTFImporter/Source/GLTFCore/Private/GLTFMeshFactory.cpp b/Engine/Plugins/Editor/GLTFImporter/Source/GLTFCore/Private/GLTFMeshFactory.cpp
--- a/Engine/Plugins/Editor/GLTFImporter/Source/GLTFCore/Private/GLTFMeshFactory.cpp
+++ b/Engine/Plugins/Editor/GLTFImporter/Source/GLTFCore/Private/GLTFMeshFactory.cpp
@@ -5,6 +5,7 @@
 
 #include "GLTFMesh.h"
 #include "StaticMeshAttributes.h"
+#include "Util/ColorConstants.h"
 
 namespace GLTF
 {
@@ -560,7 +561,10 @@ namespace GLTF
 
 				if (Colors.Num() > 0)
 				{
-					VertexInstanceColors[VertexInstanceID] = Colors[IndiceIndex];
+					FVector4f Color = Colors[IndiceIndex];
+					// Because the engine will auto convert our vertex colors to sRGB, we need to counter the conversion here to ensure the end result is still linear space
+					UE::Geometry::LinearColors::SRGBToLinear(Color);
+					VertexInstanceColors[VertexInstanceID] = Color;
 				}
 			}
 
